Media Stream
- Initial commit

#virtualized

diff --git a/Engine/Plugins/Experimental/MediaStream/MediaStream.uplugin b/Engine/Plugins/Experimental/MediaStream/MediaStream.uplugin
new file mode 100644
--- /dev/null
+++ b/Engine/Plugins/Experimental/MediaStream/MediaStream.uplugin
@@ -0,0 +1,44 @@
+{
+	"FileVersion": 3,
+	"Version": 1,
+	"VersionName": "1.0",
+	"FriendlyName": "Media Stream",
+	"Description": "Content/type agnostic chainable media proxy with media player integration.",
+	"Category": "Experimental",
+	"CreatedBy": "Epic Games, Inc.",
+	"CreatedByURL": "https://epicgames.com",
+	"DocsURL": "",
+	"MarketplaceURL": "",
+	"SupportURL": "",
+	"CanContainContent": true,
+	"IsBetaVersion": false,
+	"IsExperimentalVersion": true,
+	"Installed": false,
+	"Modules": [
+		{
+			"Name": "MediaStream",
+			"Type": "Runtime",
+			"LoadingPhase": "Default"
+		},
+		{
+			"Name": "MediaStreamEditor",
+			"Type": "Editor",
+			"LoadingPhase": "Default"
+		}
+	],
+	"Plugins": [
+
+		{
+			"Name": "LevelSequenceEditor",
+			"Enabled": true
+		},
+		{
+			"Name": "MediaCompositing",
+			"Enabled": true
+		},
+		{
+			"Name": "MediaPlayerEditor",
+			"Enabled": true
+		}
+	]
+}
diff --git a/Engine/Plugins/Experimental/MediaStream/Source/MediaStream/MediaStream.Build.cs b/Engine/Plugins/Experimental/MediaStream/Source/MediaStream/MediaStream.Build.cs
new file mode 100644
--- /dev/null
+++ b/Engine/Plugins/Experimental/MediaStream/Source/MediaStream/MediaStream.Build.cs
@@ -0,0 +1,44 @@
+// Copyright Epic Games, Inc. All Rights Reserved.
+
+using UnrealBuildTool;
+
+public class MediaStream : ModuleRules
+{
+	public MediaStream(ReadOnlyTargetRules Target) : base(Target)
+	{
+		PCHUsage = ModuleRules.PCHUsageMode.UseExplicitOrSharedPCHs;
+
+		PublicDependencyModuleNames.AddRange(
+			new string[]
+			{
+				"Core",
+				"Engine",
+				"MediaAssets"
+			}
+		);
+
+		PrivateDependencyModuleNames.AddRange(
+			new string[]
+			{
+				"CoreUObject",
+				"LevelSequence",
+				"MediaCompositing",
+				"MediaUtils",
+				"MovieScene"
+			}
+		);
+
+		if (Target.Type == TargetType.Editor)
+		{
+			PrivateDependencyModuleNames.AddRange(
+				new string[]
+				{
+					"DesktopWidgets",
+					"PropertyEditor",
+					"Slate",
+					"SlateCore"
+				}
+			);
+		}
+	}
+}
diff --git a/Engine/Plugins/Experimental/MediaStream/Source/MediaStream/Private/MediaStream.cpp b/Engine/Plugins/Experimental/MediaStream/Source/MediaStream/Private/MediaStream.cpp
new file mode 100644
--- /dev/null
+++ b/Engine/Plugins/Experimental/MediaStream/Source/MediaStream/Private/MediaStream.cpp
@@ -0,0 +1,274 @@
+// Copyright Epic Games, Inc. All Rights Reserved.
+
+#include "MediaStream.h"
+
+#include "MediaPlaylist.h"
+#include "MediaSource.h"
+#include "MediaSourceOptions.h"
+#include "MediaStreamSourceBlueprintLibrary.h"
+#include "MovieScene.h"
+#include "Players/MediaStreamLocalPlayer.h"
+#include "Players/MediaStreamProxyPlayer.h"
+
+bool UMediaStream::HasValidSource() const
+{
+	return UMediaStreamSourceBlueprintLibrary::IsValidMediaSource(Source);
+}
+
+const FMediaStreamSource& UMediaStream::GetSource() const
+{
+	return Source;
+}
+
+const FMediaStreamSource& UMediaStream::ResolveSource() const
+{
+	if (UMediaStreamProxyPlayer* ProxyPlayer = Cast<UMediaStreamProxyPlayer>(PlayerObject))
+	{
+		if (UMediaStream* ProxyStream = ProxyPlayer->GetSourceStream())
+		{
+			return ProxyStream->ResolveSource();
+		}
+	}
+
+	return Source;
+}
+
+bool UMediaStream::SetSource(const FMediaStreamSource& InSource)
+{
+	if (Source == InSource)
+	{
+		return true;
+	}
+
+	Source = InSource;
+
+	return ApplySource();
+}
+
+TScriptInterface<IMediaStreamPlayer> UMediaStream::GetPlayer() const
+{
+	return PlayerObject;
+}
+
+UMediaStream::FOnPlayerChanged& UMediaStream::GetOnPlayerChanged()
+{
+	return OnPlayerChanged;
+}
+
+#if WITH_EDITOR
+void UMediaStream::PreEditUndo()
+{
+	Super::PreEditUndo();
+
+	Source_PreUndo = Source;
+}
+
+void UMediaStream::PostEditUndo()
+{
+	Super::PostEditUndo();
+
+	if (Source_PreUndo != Source)
+	{
+		ApplySource();
+	}
+}
+
+void UMediaStream::PostEditChangeProperty(FPropertyChangedEvent& InPropertyChangedEvent)
+{
+	Super::PostEditChangeProperty(InPropertyChangedEvent);
+
+	if (InPropertyChangedEvent.GetPropertyName() == GET_MEMBER_NAME_CHECKED(UMediaStream, Source))
+	{
+		ApplySource();
+	}
+}
+#endif
+
+void UMediaStream::PostDuplicate(bool bInDuplicateForPIE)
+{
+	Super::PostDuplicate(bInDuplicateForPIE);
+
+	ApplySource();
+}
+
+void UMediaStream::PostEditImport()
+{
+	Super::PostEditImport();
+
+	ApplySource();
+}
+
+void UMediaStream::PostNetReceive()
+{
+	Super::PostNetReceive();
+
+	ApplySource();
+}
+
+void UMediaStream::BeginDestroy()
+{
+	Super::BeginDestroy();
+
+	if (IsValid(PlayerObject))
+	{
+		if (IMediaStreamPlayer* Player = Cast<IMediaStreamPlayer>(PlayerObject))
+		{
+			Player->Deinitialize();
+			PlayerObject = nullptr;
+		}
+	}
+}
+
+float UMediaStream::GetProxyRate() const
+{
+	if (IMediaStreamPlayer* MediaStreamPlayer = GetPlayer().GetInterface())
+	{
+		if (UMediaPlayer* MediaPlayer = MediaStreamPlayer->GetPlayer())
+		{
+			return MediaPlayer->GetRate();
+		}
+	}
+
+	return 0.f;
+}
+
+bool UMediaStream::SetProxyRate(float Rate)
+{
+	if (IMediaStreamPlayer* MediaStreamPlayer = GetPlayer().GetInterface())
+	{
+		if (MediaStreamPlayer->IsReadOnly())
+		{
+			return false;
+		}
+
+		if (UMediaPlayer* MediaPlayer = MediaStreamPlayer->GetPlayer())
+		{
+			return MediaPlayer->SetRate(Rate);
+		}
+	}
+
+	return false;
+}
+
+bool UMediaStream::IsExternalControlAllowed()
+{
+	if (IMediaStreamPlayer* MediaStreamPlayer = GetPlayer().GetInterface())
+	{
+		return !MediaStreamPlayer->IsReadOnly();
+	}
+
+	return false;
+}
+
+const FMediaSourceCacheSettings& UMediaStream::GetCacheSettings() const
+{
+	static FMediaSourceCacheSettings DefaultCacheSetting = {/* Override */ false, /* Cache ahead */ 0.f};
+	return DefaultCacheSetting;
+}
+
+UMediaSource* UMediaStream::ProxyGetMediaSourceFromIndex(int32 Index) const
+{
+	const FMediaStreamSource& ResolvedSource = ResolveSource();
+
+	if (UMediaPlaylist* Playlist = Cast<UMediaPlaylist>(ResolvedSource.Object))
+	{
+		return Playlist->Get(Index);
+	}
+
+	if (Index != 0)
+	{
+		return nullptr;
+	}
+
+	return Cast<UMediaSource>(ResolvedSource.Object);
+}
+
+UMediaTexture* UMediaStream::ProxyGetMediaTexture(int32 LayerIndex, int32 TextureIndex)
+{
+	if (IMediaStreamPlayer* MediaStreamPlayer = GetPlayer().GetInterface())
+	{
+		return MediaStreamPlayer->GetMediaTexture();
+	}
+
+	return nullptr;
+}
+
+void UMediaStream::ProxyReleaseMediaTexture(int32 LayerIndex, int32 TextureIndex)
+{
+	// Do nothing, we handle these internally.
+}
+
+bool UMediaStream::ProxySetAspectRatio(UMediaPlayer* InMediaPlayer)
+{
+	// We have no aspect ratio settings
+	return false;
+}
+
+void UMediaStream::ProxySetTextureBlend(int32 LayerIndex, int32 TextureIndex, float Blend)
+{
+	// Do nothing, we have no blending
+}
+
+bool UMediaStream::ApplySource()
+{
+	if (!UMediaStreamSourceBlueprintLibrary::IsValidMediaSource(Source))
+	{
+		if (IsValid(PlayerObject))
+		{
+			if (IMediaStreamPlayer* Player = Cast<IMediaStreamPlayer>(PlayerObject))
+			{
+				Player->Deinitialize();
+			}
+		}
+
+		return false;
+	}
+
+	const bool bIsProxyStream = IsValid(Source.Object) && (Source.Object->GetClass() == UMediaStream::StaticClass());
+
+	if (!IsValid(PlayerObject))
+	{
+		if (bIsProxyStream)
+		{
+			PlayerObject = NewObject<UMediaStreamProxyPlayer>(this);
+		}
+		else
+		{
+			PlayerObject = NewObject<UMediaStreamLocalPlayer>(this);
+		}
+
+		Cast<IMediaStreamPlayer>(PlayerObject)->OnCreated();
+	}
+	else
+	{
+		if (bIsProxyStream && !PlayerObject->IsA<UMediaStreamProxyPlayer>())
+		{
+			if (UMediaStreamLocalPlayer* LocalPlayer = Cast<UMediaStreamLocalPlayer>(PlayerObject))
+			{
+				LocalPlayer->Deinitialize();
+			}
+
+			PlayerObject = NewObject<UMediaStreamProxyPlayer>(this);
+			Cast<IMediaStreamPlayer>(PlayerObject)->OnCreated();
+		}
+		else if (!bIsProxyStream && !PlayerObject->IsA<UMediaStreamLocalPlayer>())
+		{
+			if (UMediaStreamProxyPlayer* ProxyPlayer = Cast<UMediaStreamProxyPlayer>(PlayerObject))
+			{
+				ProxyPlayer->Deinitialize();
+			}
+
+			PlayerObject = NewObject<UMediaStreamLocalPlayer>(this);
+			Cast<IMediaStreamPlayer>(PlayerObject)->OnCreated();
+		}
+	}
+
+	if (IMediaStreamPlayer* Player = Cast<IMediaStreamPlayer>(PlayerObject))
+	{
+		Player->OnSourceChanged(Source);
+	}
+
+	OnPlayerChanged.Broadcast(this, PlayerObject);
+
+	return true;
+}
diff --git a/Engine/Plugins/Experimental/MediaStream/Source/MediaStream/Private/MediaStreamActor.cpp b/Engine/Plugins/Experimental/MediaStream/Source/MediaStream/Private/MediaStreamActor.cpp
new file mode 100644
--- /dev/null
+++ b/Engine/Plugins/Experimental/MediaStream/Source/MediaStream/Private/MediaStreamActor.cpp
@@ -0,0 +1,11 @@
+// Copyright Epic Games, Inc. All Rights Reserved.
+
+#include "MediaStreamActor.h"
+
+#include "MediaStreamComponent.h"
+
+AMediaStreamActor::AMediaStreamActor()
+{
+	MediaStreamComponent = CreateDefaultSubobject<UMediaStreamComponent>(TEXT("MediaStreamComponent"));
+	SetRootComponent(MediaStreamComponent);
+}
diff --git a/Engine/Plugins/Experimental/MediaStream/Source/MediaStream/Private/MediaStreamComponent.cpp b/Engine/Plugins/Experimental/MediaStream/Source/MediaStream/Private/MediaStreamComponent.cpp
new file mode 100644
--- /dev/null
+++ b/Engine/Plugins/Experimental/MediaStream/Source/MediaStream/Private/MediaStreamComponent.cpp
@@ -0,0 +1,111 @@
+// Copyright Epic Games, Inc. All Rights Reserved.
+
+#include "MediaStreamComponent.h"
+
+#include "Engine/StaticMesh.h"
+#include "Engine/Texture.h"
+#include "Materials/Material.h"
+#include "Materials/MaterialInstanceDynamic.h"
+#include "IMediaStreamPlayer.h"
+#include "MediaStream.h"
+#include "UObject/Package.h"
+
+namespace UE::MediaStream::Private
+{
+	const FLazyName TextureParameterName = TEXT("MediaTexture");
+	constexpr const TCHAR* DefaultTexturePath = TEXT("/Script/Engine.Texture2D'/Engine/EditorMaterials/MatineeGroups/MAT_Groups_Anim.MAT_Groups_Anim'");
+	constexpr const TCHAR* BaseMaterialPath = TEXT("/Script/Engine.Material'/MediaStream/MediaStreamComponentMaterial.MediaStreamComponentMaterial'");
+	constexpr const TCHAR* StaticMeshPath = TEXT("/Script/Engine.StaticMesh'/Engine/EditorMeshes/EditorPlane.EditorPlane'");
+
+	UTexture* GetDefaultTexture()
+	{
+		const TSoftObjectPtr<UTexture> StaticMeshPtr = TSoftObjectPtr<UTexture>(FSoftObjectPath(DefaultTexturePath));
+		return StaticMeshPtr.LoadSynchronous();
+	}
+
+	UMaterial* GetBaseMaterial()
+	{
+		const TSoftObjectPtr<UMaterial> StaticMeshPtr = TSoftObjectPtr<UMaterial>(FSoftObjectPath(BaseMaterialPath));
+		return StaticMeshPtr.LoadSynchronous();
+	}
+
+	UStaticMesh* GetStaticMesh()
+	{
+		const TSoftObjectPtr<UStaticMesh> StaticMeshPtr = TSoftObjectPtr<UStaticMesh>(FSoftObjectPath(StaticMeshPath));
+		return StaticMeshPtr.LoadSynchronous();
+	}
+}
+
+UMediaStreamComponent::UMediaStreamComponent()
+{
+	using namespace UE::MediaStream::Private;
+
+	MediaStream = CreateDefaultSubobject<UMediaStream>("MediaStream");
+}
+
+void UMediaStreamComponent::PostInitProperties()
+{
+ 	Super::PostInitProperties();
+
+	SetStaticMesh(UE::MediaStream::Private::GetStaticMesh());
+
+	MediaStream->GetOnPlayerChanged().AddDynamic(this, &UMediaStreamComponent::OnPlayerChanged);
+
+	InitPlayer();
+}
+
+void UMediaStreamComponent::PostLoad()
+{
+	Super::PostLoad();
+
+	InitPlayer();
+}
+
+void UMediaStreamComponent::PostNetReceive()
+{
+	Super::PostNetReceive();
+
+	InitPlayer();
+}
+
+void UMediaStreamComponent::OnPlayerChanged(UMediaStream* InMediaStream, TScriptInterface<IMediaStreamPlayer> InPlayer)
+{
+	InitPlayer();
+}
+
+void UMediaStreamComponent::InitPlayer()
+{
+	using namespace UE::MediaStream::Private;
+
+	UMaterialInstanceDynamic* MID = Cast<UMaterialInstanceDynamic>(MaterialInstance);
+
+	if (!MID)
+	{
+		MaterialInstance = UMaterialInstanceDynamic::Create(GetBaseMaterial(), GetTransientPackage());
+		MID = Cast<UMaterialInstanceDynamic>(MaterialInstance);
+	}
+
+	if (GetMaterial(0) != MaterialInstance)
+	{
+		SetMaterial(0, MaterialInstance);
+		MarkRenderStateDirty();
+	}
+
+	IMediaStreamPlayer* MediaStreamPlayer = MediaStream->GetPlayer().GetInterface();
+
+	if (!MediaStreamPlayer)
+	{
+		MID->SetTextureParameterValue(TextureParameterName, GetDefaultTexture());
+		return;
+	}
+
+	UMediaTexture* MediaTexture = MediaStreamPlayer->GetMediaTexture();
+
+	if (!MediaTexture)
+	{
+		MID->SetTextureParameterValue(TextureParameterName, GetDefaultTexture());
+		return;
+	}
+
+	MID->SetTextureParameterValue(TextureParameterName, MediaTexture);
+}
diff --git a/Engine/Plugins/Experimental/MediaStream/Source/MediaStream/Private/MediaStreamModule.cpp b/Engine/Plugins/Experimental/MediaStream/Source/MediaStream/Private/MediaStreamModule.cpp
new file mode 100644
--- /dev/null
+++ b/Engine/Plugins/Experimental/MediaStream/Source/MediaStream/Private/MediaStreamModule.cpp
@@ -0,0 +1,62 @@
+// Copyright Epic Games, Inc. All Rights Reserved.
+
+#include "MediaStreamModule.h"
+
+#include "MediaStreamObjectHandlerManager.h"
+#include "MediaStreamSchemeHandlerManager.h"
+#include "Modules/ModuleManager.h"
+#include "ObjectHandlers/MediaStreamMediaPlaylistHandler.h"
+#include "ObjectHandlers/MediaStreamMediaSourceHandler.h"
+#include "ObjectHandlers/MediaStreamMediaStreamHandler.h"
+#include "SchemeHandlers/MediaStreamAssetSchemeHandler.h"
+#include "SchemeHandlers/MediaStreamFileSchemeHandler.h"
+#include "SchemeHandlers/MediaStreamManagedSchemeHandler.h"
+#include "SchemeHandlers/MediaStreamSubobjectSchemeHandler.h"
+#include "UObject/UObjectBase.h"
+
+DEFINE_LOG_CATEGORY(LogMediaStream);
+
+FMediaStreamModule& FMediaStreamModule::Get()
+{
+	return FModuleManager::Get().GetModuleChecked<FMediaStreamModule>(UE_MODULE_NAME);
+}
+
+FMediaStreamEditorPlaybackStateDelegate& FMediaStreamModule::GetEditorPlaybackStateEvent()
+{
+	return EditorPlaybackStateDelegate;
+}
+
+void FMediaStreamModule::StartupModule()
+{
+	FMediaStreamSchemeHandlerManager& SchemeHandlerManager = FMediaStreamSchemeHandlerManager::Get();
+	SchemeHandlerManager.RegisterSchemeHandler<FMediaStreamAssetSchemeHandler>();
+	SchemeHandlerManager.RegisterSchemeHandler<FMediaStreamFileSchemeHandler>();
+	SchemeHandlerManager.RegisterSchemeHandler<FMediaStreamManagedSchemeHandler>();
+	SchemeHandlerManager.RegisterSchemeHandler<FMediaStreamSubobjectSchemeHandler>();
+
+	FMediaStreamObjectHandlerManager& ObjectHandlerManager = FMediaStreamObjectHandlerManager::Get();
+	ObjectHandlerManager.RegisterObjectHandler<FMediaStreamMediaPlaylistHandler>();
+	ObjectHandlerManager.RegisterObjectHandler<FMediaStreamMediaSourceHandler>();
+	ObjectHandlerManager.RegisterObjectHandler<FMediaStreamMediaStreamHandler>();
+}
+
+void FMediaStreamModule::ShutdownModule()
+{
+	if (!UObjectInitialized())
+	{
+		return;
+	}
+
+	FMediaStreamSchemeHandlerManager& SchemeHandlerManager = FMediaStreamSchemeHandlerManager::Get();
+	SchemeHandlerManager.UnregisterSchemeHandler<FMediaStreamAssetSchemeHandler>();
+	SchemeHandlerManager.UnregisterSchemeHandler<FMediaStreamFileSchemeHandler>();
+	SchemeHandlerManager.UnregisterSchemeHandler<FMediaStreamManagedSchemeHandler>();
+	SchemeHandlerManager.UnregisterSchemeHandler<FMediaStreamSubobjectSchemeHandler>();
+
+	FMediaStreamObjectHandlerManager& ObjectHandlerManager = FMediaStreamObjectHandlerManager::Get();
+	ObjectHandlerManager.UnregisterObjectHandler<FMediaStreamMediaPlaylistHandler>();
+	ObjectHandlerManager.UnregisterObjectHandler<FMediaStreamMediaSourceHandler>();
+	ObjectHandlerManager.UnregisterObjectHandler<FMediaStreamMediaStreamHandler>();
+}
+
+IMPLEMENT_MODULE(FMediaStreamModule, MediaStream);
diff --git a/Engine/Plugins/Experimental/MediaStream/Source/MediaStream/Private/MediaStreamObjectHandlerManager.cpp b/Engine/Plugins/Experimental/MediaStream/Source/MediaStream/Private/MediaStreamObjectHandlerManager.cpp
new file mode 100644
--- /dev/null
+++ b/Engine/Plugins/Experimental/MediaStream/Source/MediaStream/Private/MediaStreamObjectHandlerManager.cpp
@@ -0,0 +1,182 @@
+// Copyright Epic Games, Inc. All Rights Reserved.
+
+#include "MediaStreamObjectHandlerManager.h"
+
+#include "IMediaAssetsModule.h"
+#include "Engine/Engine.h"
+#include "IMediaStreamObjectHandler.h"
+#include "MediaPlayer.h"
+#include "MediaPlaylist.h"
+#include "MediaSource.h"
+#include "MediaStream.h"
+#include "MediaStreamModule.h"
+#include "Modules/ModuleManager.h"
+
+FMediaStreamObjectHandlerManager& FMediaStreamObjectHandlerManager::Get()
+{
+	static FMediaStreamObjectHandlerManager Manager;
+	return Manager;
+}
+
+bool FMediaStreamObjectHandlerManager::CanHandleObject(const UClass* InClass) const
+{
+	if (!IsValid(InClass))
+	{
+		UE_LOG(LogMediaStream, Error, TEXT("Invalid Class in FMediaStreamObjectHandlerManager::CanHandleObject"));
+		return false;
+	}
+
+	return GetHandler(InClass).IsValid();
+}
+
+bool FMediaStreamObjectHandlerManager::CanHandleObject(const UObject* InObject) const
+{
+	if (!IsValid(InObject))
+	{
+		UE_LOG(LogMediaStream, Error, TEXT("Invalid Object in FMediaStreamObjectHandlerManager::CanHandleObject"));
+		return false;
+	}
+
+	return CanHandleObject(InObject->GetClass());
+}
+
+UMediaPlayer* FMediaStreamObjectHandlerManager::CreateOrUpdateMediaPlayer(UMediaStream* InMediaStream, UObject* InSourceObject, UMediaPlayer* InCurrentPlayer) const
+{
+	if (!IsValid(InMediaStream))
+	{
+		UE_LOG(LogMediaStream, Error, TEXT("Invalid Media Stream in FUMediaStreamObjectHandlerManager::CreateOrUpdateMediaPlayer"));
+		return nullptr;
+	}
+
+	if (!IsValid(InSourceObject))
+	{
+		UE_LOG(LogMediaStream, Error, TEXT("Invalid Object in FUMediaStreamObjectHandlerManager::CreateOrUpdateMediaPlayer"));
+		return nullptr;
+	}
+
+	if (InCurrentPlayer)
+	{
+		if (UMediaSource* MediaSource = Cast<UMediaSource>(InSourceObject))
+		{
+			if (InCurrentPlayer->OpenSource(MediaSource))
+			{
+				return InCurrentPlayer;
+			}
+		}
+		else if (UMediaPlaylist* Playlist = Cast<UMediaPlaylist>(InSourceObject))
+		{
+			if (InCurrentPlayer->OpenPlaylist(Playlist))
+			{
+				return InCurrentPlayer;
+			}
+		}
+	}
+
+	if (TSharedPtr<IMediaStreamObjectHandler> Handler = GetHandler(InSourceObject->GetClass()))
+	{
+		UMediaPlayer* MediaPlayer = Handler->CreateOrUpdatePlayer(InMediaStream, InSourceObject, InCurrentPlayer);
+
+		if (!MediaPlayer)
+		{
+			UE_LOG(LogMediaStream, Error, TEXT("Failed to create Media Player in FUMediaStreamObjectHandlerManager::CreateOrUpdateMediaPlayer [%s] [%s]"), *InSourceObject->GetClass()->GetName(), *InSourceObject->GetName());
+			return nullptr;
+		}
+
+		return MediaPlayer;
+	}
+
+	if (IMediaAssetsModule* MediaAssetsModule = FModuleManager::LoadModulePtr<IMediaAssetsModule>("MediaAssets"))
+	{
+		UObject* PlayerProxy = nullptr;
+		UMediaPlayer* MediaPlayer = MediaAssetsModule->GetPlayerFromObject(InSourceObject, PlayerProxy);
+
+		if (!MediaPlayer)
+		{
+			UE_LOG(LogMediaStream, Error, TEXT("No player for Object in FUMediaStreamObjectHandlerManager::CreateOrUpdateMediaPlayer [%s] [%s]"), *InSourceObject->GetClass()->GetName(), *InSourceObject->GetName());
+			return nullptr;
+		}
+
+		return MediaPlayer;
+	}
+
+	UE_LOG(LogMediaStream, Error, TEXT("Unable to get Media Assets Module in FUMediaStreamObjectHandlerManager::CreateOrUpdateMediaPlayer [%s] [%s]"), *InSourceObject->GetClass()->GetName(), *InSourceObject->GetName());
+	return nullptr;
+}
+
+bool FMediaStreamObjectHandlerManager::HasObjectHandler(const UClass* InClass) const
+{
+	if (!IsValid(InClass))
+	{
+		UE_LOG(LogMediaStream, Error, TEXT("Invalid Class in FMediaStreamObjectHandlerManager::HasObjectHandler"));
+		return false;
+	}
+
+	return Handlers.Contains(InClass->GetFName());
+}
+
+bool FMediaStreamObjectHandlerManager::HasObjectHandler(const UObject* InObject) const
+{
+	if (!IsValid(InObject))
+	{
+		UE_LOG(LogMediaStream, Error, TEXT("Invalid Object in FMediaStreamObjectHandlerManager::HasObjectHandler"));
+		return false;
+	}
+
+	return HasObjectHandler(InObject->GetClass());
+}
+
+TSharedPtr<IMediaStreamObjectHandler> FMediaStreamObjectHandlerManager::GetHandler(const UClass* InClass) const
+{
+	for (const UClass* Class = InClass; Class; Class = Class->GetSuperClass())
+	{
+		const FName ClassName = Class->GetFName();
+
+		if (const TSharedRef<IMediaStreamObjectHandler>* HandlerPtr = Handlers.Find(ClassName))
+		{
+			return *HandlerPtr;
+		}
+	}
+
+	return nullptr;
+}
+
+bool FMediaStreamObjectHandlerManager::RegisterObjectHandler(const UClass* InClass, const TSharedRef<IMediaStreamObjectHandler>& InHandler)
+{
+	if (!IsValid(InClass))
+	{
+		UE_LOG(LogMediaStream, Error, TEXT("Invalid Class in FMediaStreamObjectHandlerManager::RegisterObjectHandler"));
+		return false;
+	}
+
+	const FName ClassName = InClass->GetFName();
+
+	if (Handlers.Contains(ClassName))
+	{
+		return false;
+	}
+
+	Handlers.Add(ClassName, InHandler);
+
+	return true;
+}
+
+TSharedPtr<IMediaStreamObjectHandler> FMediaStreamObjectHandlerManager::UnregisterObjectHandler(const UClass* InClass)
+{
+	if (!IsValid(InClass))
+	{
+		UE_LOG(LogMediaStream, Error, TEXT("Invalid Class in FMediaStreamObjectHandlerManager::UnregisterObjectHandler"));
+		return nullptr;
+	}
+
+	const FName ClassName = InClass->GetFName();
+
+	if (const TSharedRef<IMediaStreamObjectHandler>* HandlerPtr = Handlers.Find(ClassName))
+	{
+		TSharedRef<IMediaStreamObjectHandler> Handler = *HandlerPtr;
+		Handlers.Remove(ClassName);
+
+		return Handler;
+	}
+
+	return nullptr;
+}
diff --git a/Engine/Plugins/Experimental/MediaStream/Source/MediaStream/Private/MediaStreamObjectHandlerSubsystem.cpp b/Engine/Plugins/Experimental/MediaStream/Source/MediaStream/Private/MediaStreamObjectHandlerSubsystem.cpp
new file mode 100644
--- /dev/null
+++ b/Engine/Plugins/Experimental/MediaStream/Source/MediaStream/Private/MediaStreamObjectHandlerSubsystem.cpp
@@ -0,0 +1,47 @@
+// Copyright Epic Games, Inc. All Rights Reserved.
+
+#include "MediaStreamObjectHandlerSubsystem.h"
+
+#include "Engine/Engine.h"
+#include "MediaStreamModule.h"
+#include "MediaStreamObjectHandlerManager.h"
+
+UMediaStreamObjectHandlerSubsystem* UMediaStreamObjectHandlerSubsystem::Get()
+{
+	if (!UObjectInitialized())
+	{
+		UE_LOG(LogMediaStream, Error, TEXT("UObject system not initialized in UMediaStreamObjectHandlerSubsystem::Get"));
+		return nullptr;
+	}
+
+	if (!GEngine)
+	{
+		UE_LOG(LogMediaStream, Error, TEXT("Invalid GEngine in UMediaStreamObjectHandlerSubsystem::Get"));
+		return nullptr;
+	}
+
+	UMediaStreamObjectHandlerSubsystem* ObjectHandlerSubsystem = GEngine->GetEngineSubsystem<UMediaStreamObjectHandlerSubsystem>();
+
+	if (!ObjectHandlerSubsystem)
+	{
+		UE_LOG(LogMediaStream, Error, TEXT("Unable to find Media Source Object Handler Subsystem in UMediaStreamObjectHandlerSubsystem::Get"));
+		return nullptr;
+	}
+
+	return ObjectHandlerSubsystem;
+}
+
+bool UMediaStreamObjectHandlerSubsystem::CanHandleObject(const UClass* InClass) const
+{
+	return FMediaStreamObjectHandlerManager::Get().CanHandleObject(InClass);
+}
+
+UMediaPlayer* UMediaStreamObjectHandlerSubsystem::CreateMediaPlayer(UMediaStream* InMediaStream, UObject* InPlaybackSource, UMediaPlayer* InCurrentPlayer) const
+{
+	return FMediaStreamObjectHandlerManager::Get().CreateOrUpdateMediaPlayer(InMediaStream, InPlaybackSource, InCurrentPlayer);
+}
+
+bool UMediaStreamObjectHandlerSubsystem::HasObjectHandler(const UClass* InClass) const
+{
+	return FMediaStreamObjectHandlerManager::Get().HasObjectHandler(InClass);
+}
diff --git a/Engine/Plugins/Experimental/MediaStream/Source/MediaStream/Private/MediaStreamPlayerConfig.cpp b/Engine/Plugins/Experimental/MediaStream/Source/MediaStream/Private/MediaStreamPlayerConfig.cpp
new file mode 100644
--- /dev/null
+++ b/Engine/Plugins/Experimental/MediaStream/Source/MediaStream/Private/MediaStreamPlayerConfig.cpp
@@ -0,0 +1,176 @@
+// Copyright Epic Games, Inc. All Rights Reserved.
+
+#include "MediaStreamPlayerConfig.h"
+
+#include "MediaPlayer.h"
+#include "MediaPlayerFacade.h"
+
+bool FMediaStreamPlayerConfig::operator==(const FMediaStreamPlayerConfig& InOther) const
+{
+	return bPlayOnOpen == InOther.bPlayOnOpen
+		&& TrackOptions.Video == InOther.TrackOptions.Video
+		&& TrackOptions.Audio == InOther.TrackOptions.Audio
+		&& TrackOptions.Caption == InOther.TrackOptions.Caption
+		&& TrackOptions.Metadata == InOther.TrackOptions.Metadata
+		&& TrackOptions.Script == InOther.TrackOptions.Script
+		&& TrackOptions.Subtitle == InOther.TrackOptions.Subtitle
+		&& TrackOptions.Text == InOther.TrackOptions.Text
+		&& bLooping == InOther.bLooping
+		&& Volume == InOther.Volume
+		&& PlaybackTimeRange == InOther.PlaybackTimeRange
+		&& Rate == InOther.Rate
+		&& TimeDelay == InOther.TimeDelay
+		&& bShuffle == InOther.bShuffle
+		&& CacheAhead == InOther.CacheAhead
+		&& CacheBehind == InOther.CacheBehind
+		&& CacheBehindGame == InOther.CacheBehindGame;
+}
+
+FMediaPlayerOptions FMediaStreamPlayerConfig::CreateOptions(const FTimespan& InStartTime, const TMap<FName, FVariant>& InCustomOptions) const
+{
+	FMediaPlayerOptions Options;
+	Options.Tracks = TrackOptions;
+	Options.TrackSelection = EMediaPlayerOptionTrackSelectMode::UseTrackOptionIndices;
+	Options.SeekTime = InStartTime;
+	Options.SeekTimeType = EMediaPlayerOptionSeekTimeType::RelativeToStartTime;
+	Options.PlayOnOpen = bPlayOnOpen ? EMediaPlayerOptionBooleanOverride::Enabled : EMediaPlayerOptionBooleanOverride::Disabled;
+	Options.Loop = bLooping ? EMediaPlayerOptionBooleanOverride::Enabled : EMediaPlayerOptionBooleanOverride::Disabled;
+	Options.InternalCustomOptions = InCustomOptions;
+
+	return Options;
+}
+
+void FMediaStreamPlayerConfig::ApplyConfig(UMediaPlayer& InMediaPlayer) const
+{
+	InMediaPlayer.PlayOnOpen = bPlayOnOpen ? 1 : 0;
+
+	auto SetTrack = [&InMediaPlayer](EMediaPlayerTrack InTrack, int32 InTrackIndex)
+		{
+			if (const int32 TrackCount = InMediaPlayer.GetNumTracks(InTrack))
+			{
+				const int32 ClampedTrackIndex = FMath::Clamp(InTrackIndex, 0, TrackCount - 1);
+
+				if (InMediaPlayer.GetSelectedTrack(InTrack) != ClampedTrackIndex)
+				{
+					InMediaPlayer.SelectTrack(InTrack, ClampedTrackIndex);
+				}
+			}
+		};
+
+	SetTrack(EMediaPlayerTrack::Video, TrackOptions.Video);
+	SetTrack(EMediaPlayerTrack::Audio, TrackOptions.Audio);
+	SetTrack(EMediaPlayerTrack::Subtitle, TrackOptions.Subtitle);
+	SetTrack(EMediaPlayerTrack::Text, TrackOptions.Text);
+	SetTrack(EMediaPlayerTrack::Caption, TrackOptions.Caption);
+	SetTrack(EMediaPlayerTrack::Script, TrackOptions.Script);
+	SetTrack(EMediaPlayerTrack::Metadata, TrackOptions.Metadata);
+
+	if (InMediaPlayer.IsLooping() != bLooping)
+	{
+		InMediaPlayer.SetLooping(bLooping);
+	}
+
+	InMediaPlayer.SetNativeVolume(FMath::Clamp(Volume, 0.f, 1.f));
+
+	if (PlaybackTimeRange.IsSet() && !FMath::IsNearlyZero(PlaybackTimeRange->Size()) && InMediaPlayer.SupportsPlaybackTimeRange())
+	{
+		const float DurationSeconds = static_cast<float>(InMediaPlayer.GetDuration().GetTotalSeconds());
+
+		if (!FMath::IsNearlyZero(DurationSeconds))
+		{
+			const float ClampedMinRange = FMath::Clamp(PlaybackTimeRange->Min, 0.f, FMath::Min(DurationSeconds, PlaybackTimeRange->Max));
+			const float ClampedMaxRange = FMath::Clamp(PlaybackTimeRange->Max, FMath::Min(DurationSeconds, PlaybackTimeRange->Min), DurationSeconds);
+
+			if (FMath::IsNearlyEqual(ClampedMinRange, ClampedMaxRange))
+			{
+				InMediaPlayer.SetPlaybackTimeRange(FFloatInterval(ClampedMinRange, ClampedMaxRange));
+			}
+		}
+	}
+
+	InMediaPlayer.PlayOnOpen = bPlayOnOpen;
+
+	if (InMediaPlayer.GetPlayerFacade()->ActivePlayerOptions.IsSet())
+	{
+		FMediaPlayerOptions& ActivePlayerOptions = InMediaPlayer.GetPlayerFacade()->ActivePlayerOptions.GetValue();
+
+		ActivePlayerOptions.PlayOnOpen = bPlayOnOpen
+			? EMediaPlayerOptionBooleanOverride::Enabled
+			: EMediaPlayerOptionBooleanOverride::Disabled;
+
+		ActivePlayerOptions.Loop = bLooping
+			? EMediaPlayerOptionBooleanOverride::Enabled
+			: EMediaPlayerOptionBooleanOverride::Disabled;
+	}
+
+	if (!bPlayOnOpen)
+	{
+		InMediaPlayer.SetRate(0.f);
+	}
+	else if (InMediaPlayer.GetRate() != Rate)
+	{
+		InMediaPlayer.Seek(0.f);
+		InMediaPlayer.PlayOnOpen = false;
+		ApplyRate(InMediaPlayer);
+	}
+
+	if (InMediaPlayer.GetTimeDelay().GetTotalSeconds() != TimeDelay)
+	{
+		InMediaPlayer.SetTimeDelay(TimeDelay);
+	}
+
+	if (InMediaPlayer.Shuffle != bShuffle)
+	{
+		InMediaPlayer.Shuffle = bShuffle ? 1 : 0;
+	}
+
+	InMediaPlayer.CacheAhead = CacheAhead;
+	InMediaPlayer.CacheBehind = CacheBehind;
+	InMediaPlayer.CacheBehindGame = CacheBehindGame;
+}
+
+bool FMediaStreamPlayerConfig::ApplyRate(UMediaPlayer& InMediaPlayer) const
+{
+ 	bool bFoundRate = false;
+	float NearestRate = 0;
+	TArray<FFloatRange> SupportedRatesList;
+	InMediaPlayer.GetSupportedRates(SupportedRatesList, /* Unthinned / no frames drops */ false);
+
+	for (const FFloatRange& SupportedRates : SupportedRatesList)
+	{
+		if (SupportedRates.Contains(Rate))
+		{
+			NearestRate = Rate;
+			bFoundRate = true;
+			break;
+		}
+
+		if (Rate < SupportedRates.GetLowerBoundValue())
+		{
+			const float BelowRange = SupportedRates.GetLowerBoundValue() - Rate;
+
+			if (!bFoundRate || BelowRange < FMath::Abs(Rate - NearestRate))
+			{
+				NearestRate = SupportedRates.GetLowerBoundValue();
+				bFoundRate = true;
+			}
+		}
+		else if (Rate > SupportedRates.GetUpperBoundValue())
+		{
+			const float AboveRange = Rate - SupportedRates.GetUpperBoundValue();
+
+			if (!bFoundRate || AboveRange < FMath::Abs(Rate - NearestRate))
+			{
+				NearestRate = SupportedRates.GetUpperBoundValue();
+				bFoundRate = true;
+			}
+		}
+	}
+
+	if (bFoundRate)
+	{
+		return InMediaPlayer.SetRate(NearestRate);
+	}
+
+	return false;
+}
diff --git a/Engine/Plugins/Experimental/MediaStream/Source/MediaStream/Private/MediaStreamSchemeHandlerLibrary.cpp b/Engine/Plugins/Experimental/MediaStream/Source/MediaStream/Private/MediaStreamSchemeHandlerLibrary.cpp
new file mode 100644
--- /dev/null
+++ b/Engine/Plugins/Experimental/MediaStream/Source/MediaStream/Private/MediaStreamSchemeHandlerLibrary.cpp
@@ -0,0 +1,90 @@
+// Copyright Epic Games, Inc. All Rights Reserved.
+
+#if WITH_EDITOR
+
+#include "IMediaStreamSchemeHandler.h"
+
+#include "MediaStreamSchemeHandlerManager.h"
+#include "MediaStreamSource.h"
+#include "PropertyHandle.h"
+
+FMediaStreamSource* FMediaStreamSchemeHandlerLibrary::GetStreamSourcePtr(TWeakPtr<IPropertyHandle> InPropertyHandleWeak)
+{
+	TSharedPtr<IPropertyHandle> PropertyHandle = InPropertyHandleWeak.Pin();
+
+	if (!PropertyHandle.IsValid())
+	{
+		return nullptr;
+	}
+
+	FProperty* Property = PropertyHandle->GetProperty();
+
+	if (!Property)
+	{
+		return nullptr;
+	}
+
+	TArray<UObject*> Outers;
+	PropertyHandle->GetOuterObjects(Outers);
+
+	if (Outers.IsEmpty())
+	{
+		return nullptr;
+	}
+
+	return Property->ContainerPtrToValuePtr<FMediaStreamSource>(Outers[0]);
+}
+
+FName FMediaStreamSchemeHandlerLibrary::GetScheme(TWeakPtr<IPropertyHandle> InPropertyHandleWeak)
+{
+	if (TSharedPtr<IPropertyHandle> PropertyHandle = InPropertyHandleWeak.Pin())
+	{
+		if (FMediaStreamSource* StreamSource = GetStreamSourcePtr(PropertyHandle.ToSharedRef()))
+		{
+			return StreamSource->Scheme;
+		}
+	}
+
+	return NAME_None;
+}
+
+FString FMediaStreamSchemeHandlerLibrary::GetPath(TWeakPtr<IPropertyHandle> InPropertyHandleWeak)
+{
+	if (TSharedPtr<IPropertyHandle> PropertyHandle = InPropertyHandleWeak.Pin())
+	{
+		if (FMediaStreamSource* StreamSource = GetStreamSourcePtr(PropertyHandle.ToSharedRef()))
+		{
+			return StreamSource->Path;
+		}
+	}
+
+	return "";
+}
+
+void FMediaStreamSchemeHandlerLibrary::SetSource(TSharedRef<IPropertyHandle> InPropertyHandle, const FName& InScheme, const FString& InPath)
+{
+	FMediaStreamSource* StreamSource = GetStreamSourcePtr(InPropertyHandle);
+
+	if (!StreamSource)
+	{
+		return;
+	}
+
+	FProperty* Property = InPropertyHandle->GetProperty();
+	TArray<UObject*> Outers;
+	InPropertyHandle->GetOuterObjects(Outers);
+
+	if (Outers.Num() != 1)
+	{
+		return;
+	}
+
+	Outers[0]->PreEditChange(Property);
+
+	*StreamSource = FMediaStreamSchemeHandlerManager::Get().CreateSource(Outers[0], InScheme, InPath);
+
+	FPropertyChangedEvent Event(Property, EPropertyChangeType::Interactive, Outers);
+	Outers[0]->PostEditChangeProperty(Event);
+}
+
+#endif
diff --git a/Engine/Plugins/Experimental/MediaStream/Source/MediaStream/Private/MediaStreamSchemeHandlerManager.cpp b/Engine/Plugins/Experimental/MediaStream/Source/MediaStream/Private/MediaStreamSchemeHandlerManager.cpp
new file mode 100644
--- /dev/null
+++ b/Engine/Plugins/Experimental/MediaStream/Source/MediaStream/Private/MediaStreamSchemeHandlerManager.cpp
@@ -0,0 +1,102 @@
+// Copyright Epic Games, Inc. All Rights Reserved.
+
+#include "MediaStreamSchemeHandlerManager.h"
+
+#include "Engine/Engine.h"
+#include "IMediaStreamSchemeHandler.h"
+#include "MediaStream.h"
+#include "MediaStreamModule.h"
+#include "MediaStreamSource.h"
+
+FMediaStreamSchemeHandlerManager& FMediaStreamSchemeHandlerManager::Get()
+{
+	static FMediaStreamSchemeHandlerManager Manager;
+	return Manager;
+}
+
+bool FMediaStreamSchemeHandlerManager::HasSchemeHandler(FName InScheme) const
+{
+	return Handlers.Contains(InScheme);
+}
+
+TArray<FName> FMediaStreamSchemeHandlerManager::GetSchemeHandlerNames() const
+{
+	TArray<FName> Keys;
+	Handlers.GenerateKeyArray(Keys);
+
+	return Keys;
+}
+
+bool FMediaStreamSchemeHandlerManager::RegisterSchemeHandler(FName InScheme, const TSharedRef<IMediaStreamSchemeHandler>& InHandler)
+{
+	if (Handlers.Contains(InScheme))
+	{
+		return false;
+	}
+
+	Handlers.Add(InScheme, InHandler);
+
+	return true;
+}
+
+TSharedPtr<IMediaStreamSchemeHandler> FMediaStreamSchemeHandlerManager::UnregisterSchemeHandler(FName InScheme)
+{
+	if (const TSharedRef<IMediaStreamSchemeHandler>* HandlerPtr = Handlers.Find(InScheme))
+	{
+		TSharedRef<IMediaStreamSchemeHandler> Handler = *HandlerPtr;
+		Handlers.Remove(InScheme);
+
+		return Handler;
+	}
+
+	return nullptr;
+}
+
+TSharedPtr<IMediaStreamSchemeHandler> FMediaStreamSchemeHandlerManager::GetHandlerTypeForScheme(FName InScheme) const
+{
+	if (InScheme.IsNone())
+	{
+		UE_LOG(LogMediaStream, Error, TEXT("Invalid Scheme in FMediaStreamSchemeHandlerManager::GetHandlerTypeForScheme"));
+		return nullptr;
+	}
+
+	if (const TSharedRef<IMediaStreamSchemeHandler>* HandlerPtr = Handlers.Find(InScheme))
+	{
+		return *HandlerPtr;
+	}
+
+	return nullptr;
+}
+
+UMediaPlayer* FMediaStreamSchemeHandlerManager::CreateOrUpdateMediaPlayer(UMediaStream* InMediaStream, UMediaPlayer* InCurrentPlayer) const
+{
+	const FMediaStreamSource& Source = InMediaStream->GetSource();
+	const TSharedRef<IMediaStreamSchemeHandler>* HandlerPtr = Handlers.Find(Source.Scheme);
+
+	if (!HandlerPtr)
+	{
+		UE_LOG(LogMediaStream, Error, TEXT("No handler for scheme in FMediaStreamSchemeHandlerManager::CreateOrUpdateMediaPlayer [%s]"), *Source.Scheme.ToString());
+		return nullptr;
+	}
+
+	return (*HandlerPtr)->CreateOrUpdateMediaPlayer(InMediaStream, InCurrentPlayer);
+}
+
+FMediaStreamSource FMediaStreamSchemeHandlerManager::CreateSource(UObject* InOuter, FName InScheme, const FString& InPath) const
+{
+	if (!IsValid(InOuter))
+	{
+		UE_LOG(LogMediaStream, Error, TEXT("Invalid Media Stream in FMediaStreamSchemeHandlerManager::GetSourceFromSchemePath"));
+		return {};
+	}
+
+	const TSharedRef<IMediaStreamSchemeHandler>* HandlerPtr = Handlers.Find(InScheme);
+
+	if (!HandlerPtr)
+	{
+		UE_LOG(LogMediaStream, Error, TEXT("Missing handler for Scheme in FMediaStreamSchemeHandlerManager::GetMediaSourceFromSchemePath [%s://%s]"), *InScheme.ToString(), *InPath);
+		return {};
+	}
+
+	return (*HandlerPtr)->CreateSource(InOuter, InPath);
+}
diff --git a/Engine/Plugins/Experimental/MediaStream/Source/MediaStream/Private/MediaStreamSchemeHandlerSubsystem.cpp b/Engine/Plugins/Experimental/MediaStream/Source/MediaStream/Private/MediaStreamSchemeHandlerSubsystem.cpp
new file mode 100644
--- /dev/null
+++ b/Engine/Plugins/Experimental/MediaStream/Source/MediaStream/Private/MediaStreamSchemeHandlerSubsystem.cpp
@@ -0,0 +1,53 @@
+// Copyright Epic Games, Inc. All Rights Reserved.
+
+#include "MediaStreamSchemeHandlerSubsystem.h"
+
+#include "Engine/Engine.h"
+#include "MediaStreamModule.h"
+#include "MediaStreamSchemeHandlerManager.h"
+#include "MediaStreamSource.h"
+
+UMediaStreamSchemeHandlerSubsystem* UMediaStreamSchemeHandlerSubsystem::Get()
+{
+	if (!UObjectInitialized())
+	{
+		UE_LOG(LogMediaStream, Error, TEXT("UObject system not initialized in UMediaStreamSchemeHandlerSubsystem::Get"));
+		return nullptr;
+	}
+
+	if (!GEngine)
+	{
+		UE_LOG(LogMediaStream, Error, TEXT("Invalid GEngine in UMediaStreamSchemeHandlerSubsystem::Get"));
+		return nullptr;
+	}
+
+	UMediaStreamSchemeHandlerSubsystem* SchemeHandlerSubsystem = GEngine->GetEngineSubsystem<UMediaStreamSchemeHandlerSubsystem>();
+
+	if (!SchemeHandlerSubsystem)
+	{
+		UE_LOG(LogMediaStream, Error, TEXT("Unable to find Media Source Url Handler Subsystem in UMediaStreamSchemeHandlerSubsystem::Get"));
+		return nullptr;
+	}
+
+	return SchemeHandlerSubsystem;
+}
+
+bool UMediaStreamSchemeHandlerSubsystem::HasSchemeHandler(FName InScheme) const
+{
+	return FMediaStreamSchemeHandlerManager::Get().HasSchemeHandler(InScheme);
+}
+
+TArray<FName> UMediaStreamSchemeHandlerSubsystem::GetSchemeHandlerNames() const
+{
+	return FMediaStreamSchemeHandlerManager::Get().GetSchemeHandlerNames();
+}
+
+UMediaPlayer* UMediaStreamSchemeHandlerSubsystem::CreateOrUpdateMediaPlayer(UMediaStream* InMediaStream, UMediaPlayer* InCurrentPlayer) const
+{
+	return FMediaStreamSchemeHandlerManager::Get().CreateOrUpdateMediaPlayer(InMediaStream, InCurrentPlayer);
+}
+
+FMediaStreamSource UMediaStreamSchemeHandlerSubsystem::CreateSource(UObject* InOuter, FName InScheme, const FString& InPath) const
+{
+	return FMediaStreamSchemeHandlerManager::Get().CreateSource(InOuter, InScheme, InPath);
+}
diff --git a/Engine/Plugins/Experimental/MediaStream/Source/MediaStream/Private/MediaStreamSource.cpp b/Engine/Plugins/Experimental/MediaStream/Source/MediaStream/Private/MediaStreamSource.cpp
new file mode 100644
--- /dev/null
+++ b/Engine/Plugins/Experimental/MediaStream/Source/MediaStream/Private/MediaStreamSource.cpp
@@ -0,0 +1,16 @@
+// Copyright Epic Games, Inc. All Rights Reserved.
+
+#include "MediaStreamSource.h"
+
+#include "MediaStreamSchemeHandlerManager.h"
+
+bool FMediaStreamSource::operator==(const FMediaStreamSource& InOther) const
+{
+	return Scheme == InOther.Scheme &&
+		(Scheme.IsNone() || (Path.Equals(InOther.Path)));
+}
+
+TArray<FName> UMediaStreamSourceBlueprintFunctionLibrary::GetSchemeTypes()
+{
+	return FMediaStreamSchemeHandlerManager::Get().GetSchemeHandlerNames();
+}
diff --git a/Engine/Plugins/Experimental/MediaStream/Source/MediaStream/Private/MediaStreamSourceBlueprintLibrary.cpp b/Engine/Plugins/Experimental/MediaStream/Source/MediaStream/Private/MediaStreamSourceBlueprintLibrary.cpp
new file mode 100644
--- /dev/null
+++ b/Engine/Plugins/Experimental/MediaStream/Source/MediaStream/Private/MediaStreamSourceBlueprintLibrary.cpp
@@ -0,0 +1,177 @@
+// Copyright Epic Games, Inc. All Rights Reserved.
+
+#include "MediaStreamSourceBlueprintLibrary.h"
+
+#include "MediaStream.h"
+#include "MediaStreamModule.h"
+#include "MediaStreamObjectHandlerManager.h"
+#include "MediaStreamSchemeHandlerManager.h"
+#include "SchemeHandlers/MediaStreamAssetSchemeHandler.h"
+#include "SchemeHandlers/MediaStreamFileSchemeHandler.h"
+#include "SchemeHandlers/MediaStreamManagedSchemeHandler.h"
+#include "SchemeHandlers/MediaStreamSubobjectSchemeHandler.h"
+
+bool UMediaStreamSourceBlueprintLibrary::IsValidMediaSource(const FMediaStreamSource& InSource)
+{
+	return !InSource.Scheme.IsNone() && !InSource.Path.IsEmpty();
+}
+
+bool UMediaStreamSourceBlueprintLibrary::IsAssetValid(const TSoftObjectPtr<UObject>& InAsset)
+{
+	return IsAssetSoftPathValid(InAsset.ToSoftObjectPath());
+}
+
+bool UMediaStreamSourceBlueprintLibrary::IsAssetPathValid(const FString& InPath)
+{
+	return IsAssetSoftPathValid(FSoftObjectPath(InPath));
+}
+
+bool UMediaStreamSourceBlueprintLibrary::IsAssetSoftPathValid(const FSoftObjectPath& InPath)
+{
+	return !InPath.IsNull() && InPath.IsValid() && InPath.IsAsset();
+}
+
+FMediaStreamSource UMediaStreamSourceBlueprintLibrary::MakeMediaSourceFromSchemePath(UMediaStream* InMediaStream, FName InScheme, const FString& InPath)
+{
+	if (!IsValid(InMediaStream))
+	{
+		UE_LOG(LogMediaStream, Error, TEXT("Invalid Media Stream in UMediaStreamSourceBlueprintLibrary::MakeMediaSourceFromSchemePath"));
+		return {};
+	}
+
+	if (InScheme.IsNone())
+	{
+		UE_LOG(LogMediaStream, Error, TEXT("Invalid Scheme in UMediaStreamSourceBlueprintLibrary::MakeMediaSourceFromSchemePath"));
+		return {};
+	}
+
+	if (InPath.IsEmpty())
+	{
+		UE_LOG(LogMediaStream, Error, TEXT("Invalid Path in UMediaStreamSourceBlueprintLibrary::MakeMediaSourceFromSchemePath"));
+		return {};
+	}
+
+	return FMediaStreamSchemeHandlerManager::Get().CreateSource(InMediaStream, InScheme, InPath);
+}
+
+FMediaStreamSource UMediaStreamSourceBlueprintLibrary::MakeMediaSourceFromAsset(UMediaStream* InMediaStream, const TSoftObjectPtr<UObject>& InObject)
+{
+	if (!IsValid(InMediaStream))
+	{
+		UE_LOG(LogMediaStream, Error, TEXT("Invalid Media Stream in UMediaStreamSourceBlueprintLibrary::MakeMediaSourceFromAsset"));
+		return {};
+	}
+
+	if (InObject.IsNull())
+	{
+		UE_LOG(LogMediaStream, Error, TEXT("Invalid Asset in UMediaStreamSourceBlueprintLibrary::MakeMediaSourceFromAsset"));
+		return {};
+	}
+
+	if (!InObject->IsAsset())
+	{
+		UE_LOG(LogMediaStream, Error, TEXT("Subobject given instead of Asset in UMediaStreamSourceBlueprintLibrary::MakeMediaSourceFromAsset"));
+		return {};
+	}
+
+	return MakeMediaSourceFromSchemePath(InMediaStream, FMediaStreamAssetSchemeHandler::Scheme, InObject.ToSoftObjectPath().ToString());
+}
+
+FMediaStreamSource UMediaStreamSourceBlueprintLibrary::MakeMediaSourceFromStreamName(UMediaStream* InMediaStream, FName InStreamName)
+{
+	if (!IsValid(InMediaStream))
+	{
+		UE_LOG(LogMediaStream, Error, TEXT("Invalid Media Stream in UMediaStreamSourceBlueprintLibrary::MakeMediaSourceFromStreamName"));
+		return {};
+	}
+
+	if (InStreamName.IsNone())
+	{
+		UE_LOG(LogMediaStream, Error, TEXT("Invalid Stream Name in UMediaStreamSourceBlueprintLibrary::MakeMediaSourceFromStreamName"));
+		return {};
+	}
+
+	return MakeMediaSourceFromSchemePath(InMediaStream, FMediaStreamManagedSchemeHandler::Scheme, InStreamName.GetPlainNameString());
+}
+
+FMediaStreamSource UMediaStreamSourceBlueprintLibrary::MakeMediaSourceFromFile(UMediaStream* InMediaStream, const FString& InFileName)
+{
+	if (!IsValid(InMediaStream))
+	{
+		UE_LOG(LogMediaStream, Error, TEXT("Invalid Media Stream in UMediaStreamSourceBlueprintLibrary::MakeMediaSourceFromFile"));
+		return {};
+	}
+
+	if (InFileName.IsEmpty())
+	{
+		UE_LOG(LogMediaStream, Error, TEXT("Invalid File Name in UMediaStreamSourceBlueprintLibrary::MakeMediaSourceFromFile"));
+		return {};
+	}
+
+	return MakeMediaSourceFromSchemePath(InMediaStream, FMediaStreamFileSchemeHandler::Scheme, InFileName);
+}
+
+FMediaStreamSource UMediaStreamSourceBlueprintLibrary::MakeMediaSourceFromSubobject(UMediaStream* InMediaStream, UObject* InObject)
+{
+	if (!IsValid(InMediaStream))
+	{
+		UE_LOG(LogMediaStream, Error, TEXT("Invalid Media Stream in UMediaStreamSourceBlueprintLibrary::MakeMediaSourceFromObject"));
+		return {};
+	}
+
+	if (!IsValid(InObject))
+	{
+		UE_LOG(LogMediaStream, Error, TEXT("Invalid Object in UMediaStreamSourceBlueprintLibrary::MakeMediaSourceFromObject"));
+		return {};
+	}
+
+	if (InObject->IsAsset())
+	{
+		UE_LOG(LogMediaStream, Error, TEXT("Asset given instead of Subobject in UMediaStreamSourceBlueprintLibrary::MakeMediaSourceFromObject [%s]"), *InObject->GetPathName());
+		return {};
+	}
+
+		const UClass* ObjectClass = InObject->GetClass();
+
+	if (!FMediaStreamObjectHandlerManager::Get().CanHandleObject(ObjectClass))
+	{
+		UE_LOG(LogMediaStream, Error, TEXT("No registered handler for class in UMediaStreamSourceBlueprintLibrary::MakeMediaSourceFromObject [%s]"), *ObjectClass->GetName());
+		return {};
+	}
+
+	const FString FullPath = InObject->GetPathName();
+	const FString RelativePath = InObject->GetPathName(InMediaStream);
+
+	if (FullPath == RelativePath)
+	{
+		UE_LOG(LogMediaStream, Error, TEXT("Subobject is not a descendant of the Root Object in UMediaStreamSourceBlueprintLibrary::MakeMediaSourceFromObject"));
+		return {};
+	}
+
+	return MakeMediaSourceFromSchemePath(InMediaStream, FMediaStreamSubobjectSchemeHandler::Scheme, InObject->GetName());
+}
+
+FMediaStreamSource UMediaStreamSourceBlueprintLibrary::MakeMediaSourceFromSubobjectClass(UMediaStream* InMediaStream, const UClass* InClass)
+{
+	if (!IsValid(InMediaStream))
+	{
+		UE_LOG(LogMediaStream, Error, TEXT("Invalid Media Stream in UMediaStreamSourceBlueprintLibrary::MakeMediaSourceFromClass"));
+		return {};
+	}
+
+	if (!IsValid(InClass))
+	{
+		UE_LOG(LogMediaStream, Error, TEXT("Invalid Class in UMediaStreamSourceBlueprintLibrary::MakeMediaSourceFromClass"));
+		return {};
+	}
+
+	if (!FMediaStreamObjectHandlerManager::Get().CanHandleObject(InClass))
+	{
+		UE_LOG(LogMediaStream, Error, TEXT("No registered handler for class in UMediaStreamSourceBlueprintLibrary::MakeMediaSourceFromClass [%s]"), *InClass->GetName());
+		return {};
+	}
+
+	UObject* Object = NewObject<UObject>(InMediaStream, InClass, NAME_None, RF_Transactional);
+
+	return MakeMediaSourceFromSubobject(InMediaStream, Object);
+}
diff --git a/Engine/Plugins/Experimental/MediaStream/Source/MediaStream/Private/MediaStreamTextureConfig.cpp b/Engine/Plugins/Experimental/MediaStream/Source/MediaStream/Private/MediaStreamTextureConfig.cpp
new file mode 100644
--- /dev/null
+++ b/Engine/Plugins/Experimental/MediaStream/Source/MediaStream/Private/MediaStreamTextureConfig.cpp
@@ -0,0 +1,21 @@
+// Copyright Epic Games, Inc. All Rights Reserved.
+
+#include "MediaStreamTextureConfig.h"
+
+bool FMediaStreamTextureConfig::operator==(const FMediaStreamTextureConfig& InOther) const
+{
+	return bEnableMipGen == InOther.bEnableMipGen
+		&& MipGenCount == InOther.MipGenCount;
+}
+
+void FMediaStreamTextureConfig::ApplyConfig(UMediaTexture& InMediaTexture) const
+{
+	if (InMediaTexture.EnableGenMips == bEnableMipGen && InMediaTexture.NumMips == MipGenCount)
+	{
+		return;
+	}
+
+	InMediaTexture.EnableGenMips = bEnableMipGen;
+	InMediaTexture.NumMips = MipGenCount;
+	InMediaTexture.UpdateResource();
+}
diff --git a/Engine/Plugins/Experimental/MediaStream/Source/MediaStream/Private/ObjectHandlers/MediaStreamMediaPlaylistHandler.cpp b/Engine/Plugins/Experimental/MediaStream/Source/MediaStream/Private/ObjectHandlers/MediaStreamMediaPlaylistHandler.cpp
new file mode 100644
--- /dev/null
+++ b/Engine/Plugins/Experimental/MediaStream/Source/MediaStream/Private/ObjectHandlers/MediaStreamMediaPlaylistHandler.cpp
@@ -0,0 +1,51 @@
+// Copyright Epic Games, Inc. All Rights Reserved.
+
+#include "ObjectHandlers/MediaStreamMediaPlaylistHandler.h"
+
+#include "MediaPlayer.h"
+#include "MediaPlaylist.h"
+#include "MediaStream.h"
+#include "MediaStreamModule.h"
+
+UClass* FMediaStreamMediaPlaylistHandler::GetClass()
+{
+	return UMediaPlaylist::StaticClass();
+}
+
+UMediaPlayer* FMediaStreamMediaPlaylistHandler::CreateOrUpdatePlayer(UMediaStream* InMediaStream, UObject* InObject, UMediaPlayer* InCurrentPlayer)
+{
+	if (!InMediaStream)
+	{
+		UE_LOG(LogMediaStream, Error, TEXT("Invalid Media Stream in FMediaStreamMediaPlaylistHandler::CreateOrUpdatePlayer"));
+		return nullptr;
+	}
+
+	if (!InObject)
+	{
+		UE_LOG(LogMediaStream, Error, TEXT("Invalid Source Object in FMediaStreamMediaPlaylistHandler::CreateOrUpdatePlayer"));
+		return nullptr;
+	}
+
+	UMediaPlaylist* MediaPlaylist = Cast<UMediaPlaylist>(InObject);
+
+	if (!MediaPlaylist)
+	{
+		UE_LOG(LogMediaStream, Error, TEXT("Invalid Media Playlist in FMediaStreamMediaPlaylistHandler::CreateOrUpdatePlayer"));
+		return nullptr;
+	}
+
+	if (InCurrentPlayer && InCurrentPlayer->OpenPlaylist(MediaPlaylist))
+	{
+		return InCurrentPlayer;
+	}
+
+	UMediaPlayer* MediaPlayer = NewObject<UMediaPlayer>(InMediaStream);
+
+	if (!MediaPlayer->OpenPlaylist(MediaPlaylist))
+	{
+		UE_LOG(LogMediaStream, Error, TEXT("Unable to create player for Media Playlist in FMediaStreamMediaSourceHandler::CreateOrUpdatePlayer"));
+		return nullptr;
+	}
+
+	return MediaPlayer;
+}
diff --git a/Engine/Plugins/Experimental/MediaStream/Source/MediaStream/Private/ObjectHandlers/MediaStreamMediaPlaylistHandler.h b/Engine/Plugins/Experimental/MediaStream/Source/MediaStream/Private/ObjectHandlers/MediaStreamMediaPlaylistHandler.h
new file mode 100644
--- /dev/null
+++ b/Engine/Plugins/Experimental/MediaStream/Source/MediaStream/Private/ObjectHandlers/MediaStreamMediaPlaylistHandler.h
@@ -0,0 +1,15 @@
+// Copyright Epic Games, Inc. All Rights Reserved.
+
+#pragma once
+
+#include "IMediaStreamObjectHandler.h"
+
+class FMediaStreamMediaPlaylistHandler : public IMediaStreamObjectHandler
+{
+public:
+	static UClass* GetClass();
+
+	//~ Begin IMediaStreamObjectHandler
+	virtual UMediaPlayer* CreateOrUpdatePlayer(UMediaStream* InMediaStream, UObject* InObject, UMediaPlayer* InCurrentPlayer) override;
+	//~ End IMediaStreamObjectHandler
+};
diff --git a/Engine/Plugins/Experimental/MediaStream/Source/MediaStream/Private/ObjectHandlers/MediaStreamMediaSourceHandler.cpp b/Engine/Plugins/Experimental/MediaStream/Source/MediaStream/Private/ObjectHandlers/MediaStreamMediaSourceHandler.cpp
new file mode 100644
--- /dev/null
+++ b/Engine/Plugins/Experimental/MediaStream/Source/MediaStream/Private/ObjectHandlers/MediaStreamMediaSourceHandler.cpp
@@ -0,0 +1,61 @@
+// Copyright Epic Games, Inc. All Rights Reserved.
+
+#include "ObjectHandlers/MediaStreamMediaSourceHandler.h"
+
+#include "IMediaStreamPlayer.h"
+#include "MediaPlayer.h"
+#include "MediaSource.h"
+#include "MediaStream.h"
+#include "MediaStreamModule.h"
+
+UClass* FMediaStreamMediaSourceHandler::GetClass()
+{
+	return UMediaSource::StaticClass();
+}
+
+UMediaPlayer* FMediaStreamMediaSourceHandler::CreateOrUpdatePlayer(UMediaStream* InMediaStream, UObject* InObject, UMediaPlayer* InCurrentPlayer)
+{
+	if (!InMediaStream)
+	{
+		UE_LOG(LogMediaStream, Error, TEXT("Invalid Media Stream in FMediaStreamMediaSourceHandler::CreateOrUpdatePlayer"));
+		return nullptr;
+	}
+
+	if (!InObject)
+	{
+		UE_LOG(LogMediaStream, Error, TEXT("Invalid Source Object in FMediaStreamMediaSourceHandler::CreateOrUpdatePlayer"));
+		return nullptr;
+	}
+
+	IMediaStreamPlayer* MediaStreamPlayer = InMediaStream->GetPlayer().GetInterface();
+
+	if (!MediaStreamPlayer)
+	{
+		UE_LOG(LogMediaStream, Error, TEXT("Invalid Media Stream Player in FMediaStreamMediaSourceHandler::CreateOrUpdatePlayer"));
+		return nullptr;
+	}
+
+	UMediaSource* MediaSource = Cast<UMediaSource>(InObject);
+
+	if (!MediaSource)
+	{
+		UE_LOG(LogMediaStream, Error, TEXT("Invalid Media Source in FMediaStreamMediaSourceHandler::CreateOrUpdatePlayer"));
+		return nullptr;
+	}
+
+	if (InCurrentPlayer && InCurrentPlayer->OpenSource(MediaSource))
+	{
+		return InCurrentPlayer;
+	}
+
+	UMediaPlayer* MediaPlayer = NewObject<UMediaPlayer>(InMediaStream, NAME_None, RF_Transactional);
+	const FMediaPlayerOptions Options = MediaStreamPlayer->GetPlayerConfig().CreateOptions(MediaStreamPlayer->GetRequestedSeekTime());
+
+	if (!MediaPlayer->OpenSourceWithOptions(MediaSource, Options))
+	{
+		UE_LOG(LogMediaStream, Error, TEXT("Unable to create player for Media Source in FMediaStreamMediaSourceHandler::CreateOrUpdatePlayer"));
+		return nullptr;
+	}
+
+	return MediaPlayer;
+}
diff --git a/Engine/Plugins/Experimental/MediaStream/Source/MediaStream/Private/ObjectHandlers/MediaStreamMediaSourceHandler.h b/Engine/Plugins/Experimental/MediaStream/Source/MediaStream/Private/ObjectHandlers/MediaStreamMediaSourceHandler.h
new file mode 100644
--- /dev/null
+++ b/Engine/Plugins/Experimental/MediaStream/Source/MediaStream/Private/ObjectHandlers/MediaStreamMediaSourceHandler.h
@@ -0,0 +1,15 @@
+// Copyright Epic Games, Inc. All Rights Reserved.
+
+#pragma once
+
+#include "IMediaStreamObjectHandler.h"
+
+class FMediaStreamMediaSourceHandler : public IMediaStreamObjectHandler
+{
+public:
+	static UClass* GetClass();
+
+	//~ Begin IMediaStreamObjectHandler
+	virtual UMediaPlayer* CreateOrUpdatePlayer(UMediaStream* InMediaStream, UObject* InObject, UMediaPlayer* InCurrentPlayer) override;
+	//~ End IMediaStreamObjectHandler
+};
diff --git a/Engine/Plugins/Experimental/MediaStream/Source/MediaStream/Private/ObjectHandlers/MediaStreamMediaStreamHandler.cpp b/Engine/Plugins/Experimental/MediaStream/Source/MediaStream/Private/ObjectHandlers/MediaStreamMediaStreamHandler.cpp
new file mode 100644
--- /dev/null
+++ b/Engine/Plugins/Experimental/MediaStream/Source/MediaStream/Private/ObjectHandlers/MediaStreamMediaStreamHandler.cpp
@@ -0,0 +1,54 @@
+// Copyright Epic Games, Inc. All Rights Reserved.
+
+#include "ObjectHandlers/MediaStreamMediaStreamHandler.h"
+
+#include "IMediaStreamPlayer.h"
+#include "MediaPlayer.h"
+#include "MediaStream.h"
+#include "MediaStreamModule.h"
+
+UClass* FMediaStreamMediaStreamHandler::GetClass()
+{
+	return UMediaStream::StaticClass();
+}
+
+UMediaPlayer* FMediaStreamMediaStreamHandler::CreateOrUpdatePlayer(UMediaStream* InMediaStream, UObject* InObject, UMediaPlayer* InCurrentPlayer)
+{
+	if (!InMediaStream)
+	{
+		UE_LOG(LogMediaStream, Error, TEXT("Invalid Media Stream in FMediaStreamMediaStreamHandler::CreateOrUpdatePlayer"));
+		return nullptr;
+	}
+
+	if (!InObject)
+	{
+		UE_LOG(LogMediaStream, Error, TEXT("Invalid Source Object in FMediaStreamMediaStreamHandler::CreateOrUpdatePlayer"));
+		return nullptr;
+	}
+
+	UMediaStream* MediaStream = Cast<UMediaStream>(InObject);
+
+	if (!MediaStream)
+	{
+		UE_LOG(LogMediaStream, Error, TEXT("Invalid Media Stream in FMediaStreamMediaStreamHandler::CreateOrUpdatePlayer"));
+		return nullptr;
+	}
+
+	IMediaStreamPlayer* MediaStreamPlayer = MediaStream->GetPlayer().GetInterface();
+
+	if (!MediaStreamPlayer)
+	{
+		UE_LOG(LogMediaStream, Error, TEXT("Invalid Media Stream Player in FMediaStreamMediaStreamHandler::CreateOrUpdatePlayer"));
+		return nullptr;
+	}
+
+	UMediaPlayer* MediaPlayer = MediaStreamPlayer->GetPlayer();
+
+	if (!MediaPlayer)
+	{
+		UE_LOG(LogMediaStream, Error, TEXT("Invalid Media Player FMediaStreamMediaStreamHandler::CreateOrUpdatePlayer"));
+		return nullptr;
+	}
+
+	return MediaPlayer;
+}
diff --git a/Engine/Plugins/Experimental/MediaStream/Source/MediaStream/Private/ObjectHandlers/MediaStreamMediaStreamHandler.h b/Engine/Plugins/Experimental/MediaStream/Source/MediaStream/Private/ObjectHandlers/MediaStreamMediaStreamHandler.h
new file mode 100644
--- /dev/null
+++ b/Engine/Plugins/Experimental/MediaStream/Source/MediaStream/Private/ObjectHandlers/MediaStreamMediaStreamHandler.h
@@ -0,0 +1,15 @@
+// Copyright Epic Games, Inc. All Rights Reserved.
+
+#pragma once
+
+#include "IMediaStreamObjectHandler.h"
+
+class FMediaStreamMediaStreamHandler : public IMediaStreamObjectHandler
+{
+public:
+	static UClass* GetClass();
+
+	//~ Begin IMediaStreamObjectHandler
+	virtual UMediaPlayer* CreateOrUpdatePlayer(UMediaStream* InMediaStream, UObject* InObject, UMediaPlayer* InCurrentPlayer) override;
+	//~ End IMediaStreamObjectHandler
+};
diff --git a/Engine/Plugins/Experimental/MediaStream/Source/MediaStream/Private/Players/MediaStreamLocalPlayer.cpp b/Engine/Plugins/Experimental/MediaStream/Source/MediaStream/Private/Players/MediaStreamLocalPlayer.cpp
new file mode 100644
--- /dev/null
+++ b/Engine/Plugins/Experimental/MediaStream/Source/MediaStream/Private/Players/MediaStreamLocalPlayer.cpp
@@ -0,0 +1,649 @@
+// Copyright Epic Games, Inc. All Rights Reserved.
+
+#include "Players/MediaStreamLocalPlayer.h"
+
+#include "CoreGlobals.h"
+#include "LevelSequence.h"
+#include "MediaPlaylist.h"
+#include "MediaStream.h"
+#include "MediaStreamModule.h"
+#include "MediaStreamSchemeHandlerManager.h"
+#include "MediaTexture.h"
+#include "MovieScene.h"
+#include "MovieSceneMediaSection.h"
+#include "MovieSceneTrack.h"
+#include "UObject/UObjectIterator.h"
+
+UMediaStreamLocalPlayer::~UMediaStreamLocalPlayer()
+{
+	if (!UObjectInitialized())
+	{
+		return;
+	}
+
+	DeinitPlayer();
+}
+
+UMediaStream* UMediaStreamLocalPlayer::GetMediaStream() const
+{
+	if (UObjectInitialized())
+	{
+		return Cast<UMediaStream>(GetOuter());
+	}
+
+	return nullptr;
+}
+
+bool UMediaStreamLocalPlayer::IsReadOnly() const
+{
+	return false;
+}
+
+void UMediaStreamLocalPlayer::OnCreated()
+{
+	MediaTexture = NewObject<UMediaTexture>(this);
+	InitTexture();
+}
+
+void UMediaStreamLocalPlayer::Deinitialize()
+{
+	DeinitPlayer();
+	DeinitTexture();
+
+#if WITH_EDITOR
+	DeinitEditorPlaybackStateEvent();
+#endif
+}
+
+#if WITH_EDITOR
+void UMediaStreamLocalPlayer::PreEditUndo()
+{
+	Super::PreEditUndo();
+
+	TextureConfig_PreUndo = TextureConfig;
+	PlayerConfig_PreUndo = PlayerConfig;
+}
+
+void UMediaStreamLocalPlayer::PostEditUndo()
+{
+	Super::PostEditUndo();
+
+	if (TextureConfig_PreUndo != TextureConfig)
+	{
+		ApplyTextureConfig();
+	}
+
+	if (PlayerConfig_PreUndo != PlayerConfig)
+	{
+		ApplyPlayerConfig();
+	}
+}
+
+void UMediaStreamLocalPlayer::PostEditChangeProperty(FPropertyChangedEvent& InPropertyChangedEvent)
+{
+	Super::PostEditChangeProperty(InPropertyChangedEvent);
+
+	static const FName PlayerConfigName = GET_MEMBER_NAME_CHECKED(UMediaStreamLocalPlayer, PlayerConfig);
+	static const FName TextureConfigName = GET_MEMBER_NAME_CHECKED(UMediaStreamLocalPlayer, TextureConfig);
+	static const FName RequestedSeekTimeName = GET_MEMBER_NAME_CHECKED(UMediaStreamLocalPlayer, RequestedSeekTime);
+	static const FName PlaybackStateName = GET_MEMBER_NAME_CHECKED(UMediaStreamLocalPlayer, PlaybackState);
+	static const FName PlaylistIndexName = GET_MEMBER_NAME_CHECKED(UMediaStreamLocalPlayer, PlaylistIndex);
+
+	const FName PropertyName = InPropertyChangedEvent.GetMemberPropertyName();
+
+	if (PropertyName == PlayerConfigName)
+	{
+		ApplyPlayerConfig();
+	}
+	else if (PropertyName == TextureConfigName)
+	{
+		ApplyTextureConfig();
+	}
+	else if (PropertyName == RequestedSeekTimeName)
+	{
+		ApplyRequestedSeekTime();
+	}
+	else if (PropertyName == PlaybackStateName)
+	{
+		ApplyPlaybackState();
+	}
+	else if (PropertyName == PlaylistIndexName)
+	{
+		ApplyPlaylistIndex();
+	}
+}
+#endif
+
+void UMediaStreamLocalPlayer::PostDuplicate(bool bInDuplicateForPIE)
+{
+	Super::PostDuplicate(bInDuplicateForPIE);
+
+#if WITH_EDITOR
+	bIsPIE = bInDuplicateForPIE;
+#endif
+
+	Initialize();
+}
+
+void UMediaStreamLocalPlayer::PostEditImport()
+{
+	Super::PostEditImport();
+
+	Initialize();
+}
+
+void UMediaStreamLocalPlayer::PostLoad()
+{
+	Super::PostLoad();
+
+	Initialize();
+}
+
+void UMediaStreamLocalPlayer::PostNetReceive()
+{
+	Super::PostNetReceive();
+
+	Initialize();
+}
+
+void UMediaStreamLocalPlayer::BeginDestroy()
+{
+	Super::BeginDestroy();
+
+	Deinitialize();
+}
+
+void UMediaStreamLocalPlayer::OnSourceChanged(const FMediaStreamSource& InSource)
+{
+#if WITH_EDITOR
+	bPlayerSourceNeedsUpdate = true;
+#endif
+
+	Initialize();
+}
+
+UMediaTexture* UMediaStreamLocalPlayer::GetMediaTexture() const
+{
+	return MediaTexture;
+}
+
+const FMediaStreamTextureConfig& UMediaStreamLocalPlayer::GetTextureConfig() const
+{
+	return TextureConfig;
+}
+
+void UMediaStreamLocalPlayer::SetTextureConfig(const FMediaStreamTextureConfig& InTextureConfig)
+{
+	TextureConfig = InTextureConfig;
+
+	ApplyTextureConfig();
+}
+
+void UMediaStreamLocalPlayer::ApplyTextureConfig()
+{
+	if (IsValid(MediaTexture))
+	{
+		TextureConfig.ApplyConfig(*MediaTexture);
+	}
+}
+
+bool UMediaStreamLocalPlayer::SetPlaylistIndex(int32 InIndex)
+{
+	PlaylistIndex = InIndex;
+
+	return ApplyPlaylistIndex();
+}
+
+UMediaPlayer* UMediaStreamLocalPlayer::GetPlayer() const
+{
+	if (IsValid(Player.Get()))
+	{
+		return Player;
+	}
+
+	return nullptr;
+}
+
+bool UMediaStreamLocalPlayer::HasValidPlayer() const
+{
+	return IsValid(Player.Get());
+}
+
+const FMediaStreamPlayerConfig& UMediaStreamLocalPlayer::GetPlayerConfig() const
+{
+	return PlayerConfig;
+}
+
+void UMediaStreamLocalPlayer::SetPlayerConfig(const FMediaStreamPlayerConfig& InPlayerConfig)
+{
+	PlayerConfig = InPlayerConfig;
+
+	ApplyPlayerConfig();
+}
+
+void UMediaStreamLocalPlayer::ApplyPlayerConfig()
+{
+	if (!IsValid(Player))
+	{
+		return;
+	}
+
+	// Stop media playing on open in the editor (only in PIE)
+#if WITH_EDITOR
+	bool bOverridePlayOnOpen = GIsEditor && !bIsPIE; // Handle -game
+	bool bInitialPlayOnOpen = PlayerConfig.bPlayOnOpen;
+
+	if (bOverridePlayOnOpen)
+	{
+		PlayerConfig.bPlayOnOpen = false;
+	}
+#endif
+
+	PlayerConfig.ApplyConfig(*Player);
+
+#if WITH_EDITOR
+	if (bOverridePlayOnOpen)
+	{
+		PlayerConfig.bPlayOnOpen = bInitialPlayOnOpen;
+	}
+#endif
+}
+
+float UMediaStreamLocalPlayer::GetRequestedSeekTime() const
+{
+	return RequestedSeekTime;
+}
+
+bool UMediaStreamLocalPlayer::SetRequestedSeekTime(float InTime)
+{
+	RequestedSeekTime = InTime;
+
+	return ApplyRequestedSeekTime();
+}
+
+bool UMediaStreamLocalPlayer::ApplyRequestedSeekTime()
+{
+	bool bSuccess = false;
+
+	if (RequestedSeekTime >= 0 && IsValid(Player) && Player->SupportsSeeking())
+	{
+		bSuccess = Player->Seek(RequestedSeekTime);
+	}
+
+	RequestedSeekTime = -1.f;
+	return bSuccess;
+}
+
+EMediaStreamPlaybackState UMediaStreamLocalPlayer::GetPlaybackState() const
+{
+	return PlaybackState;
+}
+
+bool UMediaStreamLocalPlayer::SetPlaybackState(EMediaStreamPlaybackState InState)
+{
+	PlaybackState = InState;
+
+	return ApplyPlaybackState();
+}
+
+bool UMediaStreamLocalPlayer::ApplyPlaybackState()
+{
+	if (IsValid(Player))
+	{
+		switch (PlaybackState)
+		{
+			case EMediaStreamPlaybackState::Play:
+				if (UMediaTimeStampInfo* TimeStamp = Player->GetTimeStamp())
+				{
+					const double TotalTime = Player->GetDuration().GetTotalSeconds();
+					const double CurrentTime = TimeStamp->Time.GetTotalSeconds();
+
+					if (!PlayerConfig.bLooping && FMath::IsNearlyEqual(TotalTime, CurrentTime, Player->GetVideoTrackFrameRate(INDEX_NONE, INDEX_NONE) * 0.5))
+					{
+						Player->Seek(0);
+					}
+				}
+
+				return PlayerConfig.ApplyRate(*Player);
+
+			case EMediaStreamPlaybackState::Pause:
+				return Player->SetRate(0.f);
+		}
+	}
+
+	return false;
+}
+
+int32 UMediaStreamLocalPlayer::GetPlaylistIndex() const
+{
+	return PlaylistIndex;
+}
+
+bool UMediaStreamLocalPlayer::ApplyPlaylistIndex()
+{
+	if (IsValid(Player) && PlaylistIndex >= 0)
+	{
+		if (UMediaPlaylist* Playlist = Player->GetPlaylist())
+		{
+			if (Player->GetPlaylistIndex() != PlaylistIndex)
+			{
+				return Player->OpenPlaylistIndex(Playlist, PlaylistIndex);
+			}
+
+			// Already correct
+			return true;
+		}
+	}
+
+	return false;
+}
+
+int32 UMediaStreamLocalPlayer::GetPlaylistNum() const
+{
+	if (UMediaPlaylist* Playlist = Player->GetPlaylist())
+	{
+		return Playlist->Num();
+	}
+
+	return 0;
+}
+
+UMediaStream* UMediaStreamLocalPlayer::GetSourceStream() const
+{
+	return GetMediaStream();
+}
+
+bool UMediaStreamLocalPlayer::Play()
+{
+	return SetPlaybackState(EMediaStreamPlaybackState::Play);
+}
+
+bool UMediaStreamLocalPlayer::Pause()
+{
+	return SetPlaybackState(EMediaStreamPlaybackState::Pause);
+}
+
+bool UMediaStreamLocalPlayer::Rewind()
+{
+	return SetRequestedSeekTime(0.f);
+}
+
+bool UMediaStreamLocalPlayer::FastForward()
+{
+	if (IsValid(Player))
+	{
+		return SetRequestedSeekTime(Player->GetDuration().GetTotalSeconds());
+	}
+
+	return false;
+}
+
+bool UMediaStreamLocalPlayer::Previous()
+{
+	const int32 PlaylistNum = GetPlaylistNum();
+
+	if (PlaylistNum < 1)
+	{
+		return false;
+	}
+
+	if (PlaylistNum == 1)
+	{
+		return Rewind();
+	}
+
+	const int32 NextIndex = GetPlaylistIndex() - 1;
+
+	if (NextIndex < 0)
+	{
+		if (GetPlayerConfig().bLooping)
+		{
+			return SetPlaylistIndex(PlaylistNum - 1);
+		}
+
+		return false;
+	}
+
+	return SetPlaylistIndex(NextIndex);
+}
+
+bool UMediaStreamLocalPlayer::Next()
+{
+	const int32 PlaylistNum = GetPlaylistNum();
+
+	if (PlaylistNum < 1)
+	{
+		return false;
+	}
+
+	if (PlaylistNum == 1)
+	{
+		return Rewind();
+	}
+
+	const int32 NextIndex = GetPlaylistIndex() + 1;
+
+	if (NextIndex >= PlaylistNum)
+	{
+		if (GetPlayerConfig().bLooping)
+		{
+			return SetPlaylistIndex(0);
+		}
+
+		return false;
+	}
+
+	return SetPlaylistIndex(NextIndex);
+}
+
+bool UMediaStreamLocalPlayer::Close()
+{
+	DeinitPlayer();
+	return true;
+}
+
+void UMediaStreamLocalPlayer::Initialize()
+{
+	InitTexture();
+	InitPlayer();
+
+#if WITH_EDITOR
+	InitEditorPlaybackStateEvent();
+#endif
+}
+
+void UMediaStreamLocalPlayer::InitTexture()
+{
+	if (IsValid(MediaTexture))
+	{
+		ApplyTextureConfig();
+	}
+}
+
+void UMediaStreamLocalPlayer::DeinitTexture()
+{
+}
+
+void UMediaStreamLocalPlayer::InitPlayer()
+{
+#if WITH_EDITOR
+	UMediaPlayer* CurrentPlayer = Player.Get();
+
+	if (bPlayerSourceNeedsUpdate || !IsValid(Player))
+	{
+		Player = FMediaStreamSchemeHandlerManager::Get().CreateOrUpdateMediaPlayer(GetMediaStream(), Player.Get());
+		bPlayerSourceNeedsUpdate = false;
+	}	
+
+	if (CurrentPlayer != Player)
+	{
+		UpdateSequencesWithCurrentPlayer();
+	}
+#else
+	if (!IsValid(Player))
+	{
+		return;
+	}
+	else
+	{
+		UpdateSequencesWithCurrentPlayer();
+	}
+#endif
+
+	const bool bIsPlayerValid = IsValid(Player);
+
+	if (IsValid(MediaTexture))
+	{
+		MediaTexture->SetMediaPlayer(bIsPlayerValid ? Player : nullptr);
+
+#if WITH_EDITOR
+		MediaTexture->SetDefaultMediaPlayer(bIsPlayerValid ? Player : nullptr);
+#endif
+
+		MediaTexture->UpdateResource();
+	}
+
+	if (bIsPlayerValid)
+	{
+		if (!Player->OnMediaOpened.Contains(this, GET_FUNCTION_NAME_CHECKED(UMediaStreamLocalPlayer, OnMediaOpened)))
+		{
+			Player->OnMediaOpened.AddDynamic(this, &UMediaStreamLocalPlayer::OnMediaOpened);
+		}
+
+#if WITH_EDITOR
+		Player->AffectedByPIEHandling = false;
+#endif
+
+		ApplyPlayerConfig();
+	}
+}
+
+void UMediaStreamLocalPlayer::DeinitPlayer()
+{
+	if (Player && Player->IsValidLowLevel() && IsValid(Player))
+	{
+		Player->Close();
+	}
+}
+
+void UMediaStreamLocalPlayer::OnMediaOpened(FString InOpenedUrl)
+{
+	ApplyPlayerConfig();
+}
+
+void UMediaStreamLocalPlayer::UpdateSequencesWithCurrentPlayer()
+{
+	UMediaStream* MediaStream = GetMediaStream();
+
+	if (!MediaStream)
+	{
+		return;
+	}
+
+	for (UMovieSceneMediaSection* Section : TObjectRange<UMovieSceneMediaSection>())
+	{
+		const FGuid& ProxySource = Section->GetMediaSourceProxy().GetGuid();
+
+		if (!ProxySource.IsValid())
+		{
+			continue;
+		}
+
+		UMovieSceneTrack* Track = Cast<UMovieSceneTrack>(Section->GetOuter());
+
+		if (!Track)
+		{
+			continue;
+		}
+
+		UMovieScene* MovieScene = Cast<UMovieScene>(Track->GetOuter());
+
+		if (!MovieScene)
+		{
+			continue;
+		}
+
+		ULevelSequence* LevelSequence = Cast<ULevelSequence>(MovieScene->GetOuter());
+
+		if (!LevelSequence)
+		{
+			continue;
+		}
+
+		UWorld* World = LevelSequence->GetWorld();
+
+		if (!World)
+		{
+			continue;
+		}
+
+		const FMovieSceneBindingReferences* References = LevelSequence->GetBindingReferences();
+
+		if (!References)
+		{
+			continue;
+		}
+
+		const UE::UniversalObjectLocator::FResolveParams ResolveParams(World);
+
+		TArray<UObject*, TInlineAllocator<1>> BoundObject;
+		References->ResolveBinding(ProxySource, ResolveParams, BoundObject);
+
+		if (BoundObject.IsEmpty() || BoundObject[0] != MediaStream)
+		{
+			continue;
+		}
+
+		Section->ExternalMediaPlayer = Player;
+	}
+}
+
+#if WITH_EDITOR
+void UMediaStreamLocalPlayer::InitEditorPlaybackStateEvent()
+{
+	if (!GIsEditor)
+	{
+		return;
+	}
+
+	FMediaStreamEditorPlaybackStateDelegate::RegistrationType& PIEEvent = FMediaStreamModule::Get().GetEditorPlaybackStateEvent();
+
+	if (!PIEEvent.IsBoundToObject(this))
+	{
+		PIEEvent.AddUObject(this, &UMediaStreamLocalPlayer::OnEditorPlaybackStateEvent);
+	}
+}
+
+void UMediaStreamLocalPlayer::DeinitEditorPlaybackStateEvent()
+{
+	if (!GIsEditor)
+	{
+		return;
+	}
+
+	FMediaStreamModule::Get().GetEditorPlaybackStateEvent().RemoveAll(this);
+}
+
+void UMediaStreamLocalPlayer::OnEditorPlaybackStateEvent(EMediaStreamPlaybackState InPlaybackState)
+{
+	if (bIsPIE || PlaybackState == EMediaStreamPlaybackState::Pause || !IsValid(Player))
+	{
+		return;
+	}
+
+	if (InPlaybackState == EMediaStreamPlaybackState::Play)
+	{
+		if (!FMath::IsNearlyZero(PlayerConfig.Rate))
+		{
+			Player->SetRate(PlayerConfig.Rate);
+		}
+	}
+	else if (InPlaybackState == EMediaStreamPlaybackState::Pause)
+	{
+		if (!FMath::IsNearlyZero(Player->GetRate()))
+		{
+			Player->Pause();
+		}
+	}
+}
+#endif
diff --git a/Engine/Plugins/Experimental/MediaStream/Source/MediaStream/Private/Players/MediaStreamLocalPlayer.h b/Engine/Plugins/Experimental/MediaStream/Source/MediaStream/Private/Players/MediaStreamLocalPlayer.h
new file mode 100644
--- /dev/null
+++ b/Engine/Plugins/Experimental/MediaStream/Source/MediaStream/Private/Players/MediaStreamLocalPlayer.h
@@ -0,0 +1,194 @@
+// Copyright Epic Games, Inc. All Rights Reserved.
+
+#pragma once
+
+#include "IMediaStreamPlayer.h"
+#include "UObject/Object.h"
+
+#include "Misc/Guid.h"
+#include "UObject/WeakObjectPtr.h"
+
+#include "MediaStreamLocalPlayer.generated.h"
+
+/**
+ * Media Stream Local Player. Plays a media source using a UMediaPlayer.
+ */
+UCLASS(BlueprintType, Blueprintable, ClassGroup = "Media Stream", DefaultToInstanced)
+class UMediaStreamLocalPlayer : public UObject, public IMediaStreamPlayer
+{
+	GENERATED_BODY()
+
+public:
+	UMediaStreamLocalPlayer() = default;
+	virtual ~UMediaStreamLocalPlayer() override;	
+
+	UFUNCTION(BlueprintCallable, DisplayName = "Set Requested Seek Time", Category = "Media Stream|Player")
+	void BP_SetRequestedSeekTime(float InSeekTime)
+	{
+		SetRequestedSeekTime(InSeekTime);
+	}
+
+	UFUNCTION(BlueprintCallable, DisplayName = "Set Playback State", Category = "Media Stream|Player")
+	void BP_SetPlaybackState(EMediaStreamPlaybackState InState)
+	{
+		SetPlaybackState(InState);
+	}
+
+	UFUNCTION(BlueprintCallable, DisplayName = "Set Playlist Index", Category = "Media Stream|Player")
+	void BP_SetPlaylistIndex(int32 InIndex)
+	{
+		SetPlaylistIndex(InIndex);
+	}
+
+	//~ Begin IMediaStreamPlayer
+	virtual UMediaStream* GetMediaStream() const override;
+	virtual bool IsReadOnly() const override;
+	virtual void OnCreated() override;
+	virtual void OnSourceChanged(const FMediaStreamSource& InSource) override;
+	virtual UMediaTexture* GetMediaTexture() const override;
+	virtual const FMediaStreamTextureConfig& GetTextureConfig() const override;
+	virtual void SetTextureConfig(const FMediaStreamTextureConfig& InTextureConfig) override;
+	virtual void ApplyTextureConfig() override;
+	virtual UMediaPlayer* GetPlayer() const override;
+	virtual bool HasValidPlayer() const override;
+	virtual const FMediaStreamPlayerConfig& GetPlayerConfig() const override;
+	virtual void SetPlayerConfig(const FMediaStreamPlayerConfig& InPlayerConfig) override;
+	virtual void ApplyPlayerConfig() override;
+	virtual bool SetPlaylistIndex(int32 InIndex) override;
+	virtual float GetRequestedSeekTime() const override;
+	virtual bool SetRequestedSeekTime(float InTime) override;
+	virtual EMediaStreamPlaybackState GetPlaybackState() const override;
+	virtual bool SetPlaybackState(EMediaStreamPlaybackState InState) override;
+	virtual int32 GetPlaylistIndex() const override;
+	virtual int32 GetPlaylistNum() const override;
+	virtual UMediaStream* GetSourceStream() const override;
+	virtual bool Play() override;
+	virtual bool Pause() override;
+	virtual bool Rewind() override;
+	virtual bool FastForward() override;
+	virtual bool Previous() override;
+	virtual bool Next() override;
+	virtual bool Close() override;
+	virtual void Deinitialize() override;
+	//~ End IMediaStreamPlayer
+
+	//~ Begin UObject
+#if WITH_EDITOR
+	virtual void PreEditUndo() override;
+	virtual void PostEditUndo() override;
+	virtual void PostEditChangeProperty(FPropertyChangedEvent& InPropertyChangedEvent) override;
+#endif
+
+	virtual void PostDuplicate(bool bInDuplicateForPIE) override;
+	virtual void PostEditImport() override;
+	virtual void PostLoad() override;
+	virtual void PostNetReceive() override;
+	virtual void BeginDestroy() override;
+	//~ End UObject
+
+protected:
+	UPROPERTY(VisibleInstanceOnly, BlueprintReadOnly, Category = "Media Stream|Texture", meta = (AllowPrivateAccess = true))
+	TObjectPtr<UMediaTexture> MediaTexture;
+
+	UPROPERTY(EditAnywhere, BlueprintReadWrite, Category = "Media Stream|Texture", meta = (AllowPrivateAccess = true))
+	FMediaStreamTextureConfig TextureConfig;
+
+	UPROPERTY(BlueprintReadOnly, Category = "Media Stream|Player", Transient, DuplicateTransient, TextExportTransient,
+		meta = (AllowPrivateAccess = true))
+	TObjectPtr<UMediaPlayer> Player = nullptr;
+
+	UPROPERTY(EditAnywhere, BlueprintReadWrite, Category = "Media Stream|Player", meta = (AllowPrivateAccess = true))
+	FMediaStreamPlayerConfig PlayerConfig;
+
+	/** Attempts to seek to this time. Any value below 0 is ignored. */
+	UPROPERTY(EditAnywhere, BlueprintReadWrite, Setter = BP_SetRequestedSeekTime, Getter, BlueprintSetter = BP_SetRequestedSeekTime,
+		Category = "Media Stream|Player", Interp, meta = (AllowPrivateAccess = true))
+	float RequestedSeekTime = -1.f;
+
+	UPROPERTY(EditAnywhere, BlueprintReadWrite, Setter = BP_SetPlaybackState, Getter, BlueprintSetter = BP_SetPlaybackState,
+		Category = "Media Stream|Player", Interp, meta = (AllowPrivateAccess = true))
+	EMediaStreamPlaybackState PlaybackState = EMediaStreamPlaybackState::Play;
+
+	/** The index playing in the playlist. -1 to ignore. */
+	UPROPERTY(EditAnywhere, BlueprintReadWrite, Setter = BP_SetPlaylistIndex, Getter, BlueprintSetter = BP_SetPlaylistIndex,
+		Category = "Media Stream|Player", Interp, meta = (AllowPrivateAccess = true))
+	int32 PlaylistIndex = -1;
+
+	FMediaStreamTextureConfig TextureConfig_PreUndo;
+	FMediaStreamPlayerConfig PlayerConfig_PreUndo;
+
+#if WITH_EDITOR
+	bool bIsPIE = false;
+	bool bPlayerSourceNeedsUpdate = false;
+#endif
+
+	/**
+	 * Initializes all of the components of the Media Stream.
+	 */
+	void Initialize();
+
+	/**
+	 * Initializes the Media Texture
+	 */
+	void InitTexture();
+
+	/**
+	 * Deinitializes the Media Texture
+	 */
+	void DeinitTexture();
+
+	/**
+	 * Initializes the Media Player.
+	 */
+	void InitPlayer();
+
+	/**
+	 * Deinitializes the Media Player and clears the object.
+	 */
+	void DeinitPlayer();
+
+	/**
+	 * Attempts to set the seek time on the currently active Media Player.
+	 * @return True if successful.
+	 */
+	bool ApplyRequestedSeekTime();
+
+	/**
+	 * Will either set the current Rate from the player config or 0 for paused.
+	 * The player will be checked for valid rate values. The nearest possible rate will be chosen.
+	 * @return True if successful.
+	 */
+	bool ApplyPlaybackState();
+
+	/**
+	 * Attempts to set the playlist index on the Media Player.
+	 * @return True if successful.
+	 */
+	bool ApplyPlaylistIndex();
+
+	/**
+	 * Called when the player's media has opened.
+	 */
+	UFUNCTION()
+	void OnMediaOpened(FString InOpenedUrl);
+
+	void UpdateSequencesWithCurrentPlayer();
+
+#if WITH_EDITOR
+	/**
+	 * Initializes PIE events
+	 */
+	void InitEditorPlaybackStateEvent();
+
+	/**
+	 * Deinitializes PIE events
+	 */
+	void DeinitEditorPlaybackStateEvent();
+
+	/**
+	 * Called when the editor button to start and stop playback in the editor is pressed.
+	 * @param InPlaybackState
+	 */
+	void OnEditorPlaybackStateEvent(EMediaStreamPlaybackState InPlaybackState);
+#endif
+};
diff --git a/Engine/Plugins/Experimental/MediaStream/Source/MediaStream/Private/Players/MediaStreamProxyPlayer.cpp b/Engine/Plugins/Experimental/MediaStream/Source/MediaStream/Private/Players/MediaStreamProxyPlayer.cpp
new file mode 100644
--- /dev/null
+++ b/Engine/Plugins/Experimental/MediaStream/Source/MediaStream/Private/Players/MediaStreamProxyPlayer.cpp
@@ -0,0 +1,417 @@
+// Copyright Epic Games, Inc. All Rights Reserved.
+
+#include "Players/MediaStreamProxyPlayer.h"
+
+#include "MediaStream.h"
+#include "SchemeHandlers/MediaStreamAssetSchemeHandler.h"
+#include "UObject/ScriptInterface.h"
+
+UMediaStreamProxyPlayer::~UMediaStreamProxyPlayer()
+{
+	Deinitialize();
+}
+
+TSoftObjectPtr<UMediaStream> UMediaStreamProxyPlayer::GetProxyStreamSoft() const
+{
+	return ProxyStreamSoft;
+}
+
+UMediaStream* UMediaStreamProxyPlayer::GetSourceStream() const
+{
+	// Skip this check if it's already loaded.
+	if (ProxyStream != nullptr && ProxyStream == ProxyStreamSoft.Get())
+	{
+		return ProxyStream;
+	}
+
+	ProxyStream = nullptr;
+
+	if (!ProxyStreamSoft.IsNull())
+	{
+		ProxyStream = ProxyStreamSoft.LoadSynchronous();
+	}
+
+	return ProxyStream;
+}
+
+bool UMediaStreamProxyPlayer::IsReadOnly() const
+{
+	return bReadOnly;
+}
+
+void UMediaStreamProxyPlayer::SetReadOnly(bool bInReadOnly)
+{
+	bReadOnly = bInReadOnly;
+}
+
+UMediaStream* UMediaStreamProxyPlayer::GetMediaStream() const
+{
+	if (UObjectInitialized())
+	{
+		return Cast<UMediaStream>(GetOuter());
+	}
+
+	return nullptr;
+}
+
+void UMediaStreamProxyPlayer::Deinitialize()
+{
+	ProxyStreamSoft.Reset();
+	ProxyStream = nullptr;
+}
+
+void UMediaStreamProxyPlayer::OnSourceChanged(const FMediaStreamSource& InSource)
+{
+	ProxyStreamSoft.Reset();
+	ProxyStream = nullptr;
+
+	if (InSource.Scheme == FMediaStreamAssetSchemeHandler::Scheme)
+	{
+		if (IsValid(InSource.Object) && InSource.Object->GetClass() == UMediaStream::StaticClass())
+		{
+			// Lets us "cast" this soft pointer.
+			ProxyStreamSoft = Cast<UMediaStream>(InSource.Object);
+		}
+	}
+}
+
+UMediaTexture* UMediaStreamProxyPlayer::GetMediaTexture() const
+{
+	if (UMediaStream* MediaStream = GetSourceStream())
+	{
+		if (TScriptInterface<IMediaStreamPlayer> MediaStreamPlayer = MediaStream->GetPlayer())
+		{
+			return MediaStreamPlayer->GetMediaTexture();
+		}
+	}
+
+	return nullptr;
+}
+
+const FMediaStreamTextureConfig& UMediaStreamProxyPlayer::GetTextureConfig() const
+{
+	if (UMediaStream* MediaStream = GetSourceStream())
+	{
+		if (TScriptInterface<IMediaStreamPlayer> MediaStreamPlayer = MediaStream->GetPlayer())
+		{
+			return MediaStreamPlayer->GetTextureConfig();
+		}
+	}
+
+	static const FMediaStreamTextureConfig Config;
+
+	return Config;
+}
+
+void UMediaStreamProxyPlayer::SetTextureConfig(const FMediaStreamTextureConfig& InTextureConfig)
+{
+	if (!bReadOnly)
+	{
+		if (UMediaStream* MediaStream = GetSourceStream())
+		{
+			if (TScriptInterface<IMediaStreamPlayer> MediaStreamPlayer = MediaStream->GetPlayer())
+			{
+				MediaStreamPlayer->SetTextureConfig(InTextureConfig);
+			}
+		}
+	}
+}
+
+void UMediaStreamProxyPlayer::ApplyTextureConfig()
+{
+	if (!bReadOnly)
+	{
+		if (UMediaStream* MediaStream = GetSourceStream())
+		{
+			if (TScriptInterface<IMediaStreamPlayer> MediaStreamPlayer = MediaStream->GetPlayer())
+			{
+				MediaStreamPlayer->ApplyTextureConfig();
+			}
+		}
+	}
+}
+
+void UMediaStreamProxyPlayer::OnCreated()
+{
+}
+
+bool UMediaStreamProxyPlayer::SetPlaylistIndex(int32 InIndex)
+{
+	if (!bReadOnly)
+	{
+		if (UMediaStream* MediaStream = GetSourceStream())
+		{
+			if (TScriptInterface<IMediaStreamPlayer> MediaStreamPlayer = MediaStream->GetPlayer())
+			{
+				return MediaStreamPlayer->SetPlaylistIndex(InIndex);
+			}
+		}
+	}
+
+	return false;
+}
+
+UMediaPlayer* UMediaStreamProxyPlayer::GetPlayer() const
+{
+	if (UMediaStream* MediaStream = GetSourceStream())
+	{
+		if (TScriptInterface<IMediaStreamPlayer> MediaStreamPlayer = MediaStream->GetPlayer())
+		{
+			return MediaStreamPlayer->GetPlayer();
+		}
+	}
+
+	return nullptr;
+}
+
+bool UMediaStreamProxyPlayer::HasValidPlayer() const
+{
+	if (UMediaStream* MediaStream = GetSourceStream())
+	{
+		if (TScriptInterface<IMediaStreamPlayer> MediaStreamPlayer = MediaStream->GetPlayer())
+		{
+			return MediaStreamPlayer->HasValidPlayer();
+		}
+	}
+
+	return false;
+}
+
+const FMediaStreamPlayerConfig& UMediaStreamProxyPlayer::GetPlayerConfig() const
+{
+	if (UMediaStream* MediaStream = GetSourceStream())
+	{
+		if (TScriptInterface<IMediaStreamPlayer> MediaStreamPlayer = MediaStream->GetPlayer())
+		{
+			return MediaStreamPlayer->GetPlayerConfig();
+		}
+	}
+
+	static const FMediaStreamPlayerConfig Config;
+
+	return Config;
+}
+
+void UMediaStreamProxyPlayer::SetPlayerConfig(const FMediaStreamPlayerConfig& InPlayerConfig)
+{
+	if (!bReadOnly)
+	{
+		if (UMediaStream* MediaStream = GetSourceStream())
+		{
+			if (TScriptInterface<IMediaStreamPlayer> MediaStreamPlayer = MediaStream->GetPlayer())
+			{
+				return MediaStreamPlayer->SetPlayerConfig(InPlayerConfig);
+			}
+		}
+	}
+}
+
+void UMediaStreamProxyPlayer::ApplyPlayerConfig()
+{
+	if (!bReadOnly)
+	{
+		if (UMediaStream* MediaStream = GetSourceStream())
+		{
+			if (TScriptInterface<IMediaStreamPlayer> MediaStreamPlayer = MediaStream->GetPlayer())
+			{
+				MediaStreamPlayer->ApplyPlayerConfig();
+			}
+		}
+	}
+}
+
+float UMediaStreamProxyPlayer::GetRequestedSeekTime() const
+{
+	if (UMediaStream* MediaStream = GetSourceStream())
+	{
+		if (TScriptInterface<IMediaStreamPlayer> MediaStreamPlayer = MediaStream->GetPlayer())
+		{
+			return MediaStreamPlayer->GetRequestedSeekTime();
+		}
+	}
+
+	return 0.f;
+}
+
+bool UMediaStreamProxyPlayer::SetRequestedSeekTime(float InTime)
+{
+	if (!bReadOnly)
+	{
+		if (UMediaStream* MediaStream = GetSourceStream())
+		{
+			if (TScriptInterface<IMediaStreamPlayer> MediaStreamPlayer = MediaStream->GetPlayer())
+			{
+				return MediaStreamPlayer->SetRequestedSeekTime(InTime);
+			}
+		}
+	}
+
+	return false;
+}
+
+EMediaStreamPlaybackState UMediaStreamProxyPlayer::GetPlaybackState() const
+{
+	if (UMediaStream* MediaStream = GetSourceStream())
+	{
+		if (TScriptInterface<IMediaStreamPlayer> MediaStreamPlayer = MediaStream->GetPlayer())
+		{
+			return MediaStreamPlayer->GetPlaybackState();
+		}
+	}
+
+	return EMediaStreamPlaybackState::Play;
+}
+
+bool UMediaStreamProxyPlayer::SetPlaybackState(EMediaStreamPlaybackState InState)
+{
+	if (!bReadOnly)
+	{
+		if (UMediaStream* MediaStream = GetSourceStream())
+		{
+			if (TScriptInterface<IMediaStreamPlayer> MediaStreamPlayer = MediaStream->GetPlayer())
+			{
+				return MediaStreamPlayer->SetPlaybackState(InState);
+			}
+		}
+	}
+
+	return false;
+}
+
+int32 UMediaStreamProxyPlayer::GetPlaylistIndex() const
+{
+	if (UMediaStream* MediaStream = GetSourceStream())
+	{
+		if (TScriptInterface<IMediaStreamPlayer> MediaStreamPlayer = MediaStream->GetPlayer())
+		{
+			return MediaStreamPlayer->GetPlaylistIndex();
+		}
+	}
+
+	return -1;
+}
+
+int32 UMediaStreamProxyPlayer::GetPlaylistNum() const
+{
+	if (UMediaStream* MediaStream = GetSourceStream())
+	{
+		if (TScriptInterface<IMediaStreamPlayer> MediaStreamPlayer = MediaStream->GetPlayer())
+		{
+			return MediaStreamPlayer->GetPlaylistNum();
+		}
+	}
+
+	return -1;
+}
+
+bool UMediaStreamProxyPlayer::Play()
+{
+	if (!bReadOnly)
+	{
+		if (UMediaStream* MediaStream = GetSourceStream())
+		{
+			if (TScriptInterface<IMediaStreamPlayer> MediaStreamPlayer = MediaStream->GetPlayer())
+			{
+				return MediaStreamPlayer->Play();
+			}
+		}
+	}
+
+	return false;
+}
+
+bool UMediaStreamProxyPlayer::Pause()
+{
+	if (!bReadOnly)
+	{
+		if (UMediaStream* MediaStream = GetSourceStream())
+		{
+			if (TScriptInterface<IMediaStreamPlayer> MediaStreamPlayer = MediaStream->GetPlayer())
+			{
+				return MediaStreamPlayer->Pause();
+			}
+		}
+	}
+
+	return false;
+}
+
+bool UMediaStreamProxyPlayer::Rewind()
+{
+	if (!bReadOnly)
+	{
+		if (UMediaStream* MediaStream = GetSourceStream())
+		{
+			if (TScriptInterface<IMediaStreamPlayer> MediaStreamPlayer = MediaStream->GetPlayer())
+			{
+				return MediaStreamPlayer->Rewind();
+			}
+		}
+	}
+
+	return false;
+}
+
+bool UMediaStreamProxyPlayer::FastForward()
+{
+	if (!bReadOnly)
+	{
+		if (UMediaStream* MediaStream = GetSourceStream())
+		{
+			if (TScriptInterface<IMediaStreamPlayer> MediaStreamPlayer = MediaStream->GetPlayer())
+			{
+				return MediaStreamPlayer->FastForward();
+			}
+		}
+	}
+
+	return false;
+}
+
+bool UMediaStreamProxyPlayer::Previous()
+{
+	if (!bReadOnly)
+	{
+		if (UMediaStream* MediaStream = GetSourceStream())
+		{
+			if (TScriptInterface<IMediaStreamPlayer> MediaStreamPlayer = MediaStream->GetPlayer())
+			{
+				return MediaStreamPlayer->Previous();
+			}
+		}
+	}
+
+	return false;
+}
+
+bool UMediaStreamProxyPlayer::Next()
+{
+	if (!bReadOnly)
+	{
+		if (UMediaStream* MediaStream = GetSourceStream())
+		{
+			if (TScriptInterface<IMediaStreamPlayer> MediaStreamPlayer = MediaStream->GetPlayer())
+			{
+				return MediaStreamPlayer->Next();
+			}
+		}
+	}
+
+	return false;
+}
+
+bool UMediaStreamProxyPlayer::Close()
+{
+	if (!bReadOnly)
+	{
+		if (UMediaStream* MediaStream = GetSourceStream())
+		{
+			if (TScriptInterface<IMediaStreamPlayer> MediaStreamPlayer = MediaStream->GetPlayer())
+			{
+				return MediaStreamPlayer->Close();
+			}
+		}
+	}
+
+	return false;
+}
diff --git a/Engine/Plugins/Experimental/MediaStream/Source/MediaStream/Private/Players/MediaStreamProxyPlayer.h b/Engine/Plugins/Experimental/MediaStream/Source/MediaStream/Private/Players/MediaStreamProxyPlayer.h
new file mode 100644
--- /dev/null
+++ b/Engine/Plugins/Experimental/MediaStream/Source/MediaStream/Private/Players/MediaStreamProxyPlayer.h
@@ -0,0 +1,71 @@
+// Copyright Epic Games, Inc. All Rights Reserved.
+
+#pragma once
+
+#include "IMediaStreamPlayer.h"
+#include "UObject/Object.h"
+
+#include "MediaStreamProxyPlayer.generated.h"
+
+/**
+ * Media Stream Proxy Player. Forwards the player and texture from another Media Stream.
+ */
+UCLASS(BlueprintType, Blueprintable, ClassGroup = "Media Stream", DefaultToInstanced)
+class UMediaStreamProxyPlayer : public UObject, public IMediaStreamPlayer
+{
+	GENERATED_BODY()
+
+public:
+	virtual ~UMediaStreamProxyPlayer() override;
+
+	UFUNCTION(BlueprintPure, Category = "Media Stream|Player")
+	TSoftObjectPtr<UMediaStream> GetProxyStreamSoft() const;
+
+	UFUNCTION(BlueprintCallable, Category = "Media Stream|Texture")
+	void SetReadOnly(bool bInReadOnly);
+
+	//~ Begin IMediaStreamPlayer
+	virtual UMediaStream* GetMediaStream() const override;
+	virtual bool IsReadOnly() const override;
+	virtual void OnCreated() override;
+	virtual void OnSourceChanged(const FMediaStreamSource& InSource) override;
+	virtual UMediaTexture* GetMediaTexture() const override;
+	virtual const FMediaStreamTextureConfig& GetTextureConfig() const override;
+	virtual void SetTextureConfig(const FMediaStreamTextureConfig& InTextureConfig) override;
+	virtual void ApplyTextureConfig() override;
+	virtual UMediaPlayer* GetPlayer() const override;
+	virtual bool HasValidPlayer() const override;
+	virtual const FMediaStreamPlayerConfig& GetPlayerConfig() const override;
+	virtual void SetPlayerConfig(const FMediaStreamPlayerConfig& InPlayerConfig) override;
+	virtual void ApplyPlayerConfig() override;
+	virtual bool SetPlaylistIndex(int32 InIndex) override;
+	virtual float GetRequestedSeekTime() const override;
+	virtual bool SetRequestedSeekTime(float InTime) override;
+	virtual EMediaStreamPlaybackState GetPlaybackState() const override;
+	virtual bool SetPlaybackState(EMediaStreamPlaybackState InState) override;
+	virtual int32 GetPlaylistIndex() const override;
+	virtual int32 GetPlaylistNum() const override;
+	virtual UMediaStream* GetSourceStream() const override;
+	virtual bool Play() override;
+	virtual bool Pause() override;
+	virtual bool Rewind() override;
+	virtual bool FastForward() override;
+	virtual bool Previous() override;
+	virtual bool Next() override;
+	virtual bool Close() override;
+	virtual void Deinitialize() override;
+	//~ End IMediaStreamPlayer
+
+protected:
+	/** Soft pointer to the proxied stream. */
+	UPROPERTY()
+	TSoftObjectPtr<UMediaStream> ProxyStreamSoft;
+
+	/** Hard reference to the proxied stream. */
+	UPROPERTY(Transient)
+	mutable TObjectPtr<UMediaStream> ProxyStream;
+
+	/** If true, the proxied stream won't receive set calls. */
+	UPROPERTY(EditAnywhere, BlueprintReadWrite, Category = "Media Stream|Player")
+	bool bReadOnly = true;
+};
diff --git a/Engine/Plugins/Experimental/MediaStream/Source/MediaStream/Private/SchemeHandlers/MediaStreamAssetSchemeHandler.cpp b/Engine/Plugins/Experimental/MediaStream/Source/MediaStream/Private/SchemeHandlers/MediaStreamAssetSchemeHandler.cpp
new file mode 100644
--- /dev/null
+++ b/Engine/Plugins/Experimental/MediaStream/Source/MediaStream/Private/SchemeHandlers/MediaStreamAssetSchemeHandler.cpp
@@ -0,0 +1,168 @@
+// Copyright Epic Games, Inc. All Rights Reserved.
+
+#include "SchemeHandlers/MediaStreamAssetSchemeHandler.h"
+
+#include "MediaSource.h"
+#include "MediaStream.h"
+#include "MediaStreamModule.h"
+#include "MediaStreamObjectHandlerManager.h"
+#include "MediaStreamSource.h"
+#include "MediaStreamSourceBlueprintLibrary.h"
+
+#if WITH_EDITOR
+#include "DetailLayoutBuilder.h"
+#include "DetailWidgetRow.h"
+#include "IDetailGroup.h"
+#include "PropertyCustomizationHelpers.h"
+#include "Widgets/Text/STextBlock.h"
+#endif
+
+#define LOCTEXT_NAMESPACE "MediaStreamAssetSchemeHandler"
+
+const FLazyName FMediaStreamAssetSchemeHandler::Scheme = TEXT("Asset");
+
+FMediaStreamSource FMediaStreamAssetSchemeHandler::CreateSource(UObject* InOuter, const FString& InPath)
+{
+	FMediaStreamSource Source;
+
+	if (!IsValid(InOuter))
+	{
+		UE_LOG(LogMediaStream, Error, TEXT("Invalid Media Stream in FMediaStreamAssetSchemeHandler::CreateSource"));
+		return Source;
+	}
+
+	Source.Scheme = Scheme;
+
+	UObject* Asset = ResolveAsset(InPath);
+
+	if (!IsValid(Asset))
+	{
+		return Source;
+	}
+
+	UMediaSource* MediaSource = Cast<UMediaSource>(Asset);
+
+	if (!MediaSource)
+	{
+		return Source;
+	}
+
+	Source.Path = InPath;
+	Source.Object = MediaSource;
+
+	return Source;
+}
+
+UMediaPlayer* FMediaStreamAssetSchemeHandler::CreateOrUpdateMediaPlayer(UMediaStream* InMediaStream, UMediaPlayer* InCurrentPlayer)
+{
+	if (!IsValid(InMediaStream))
+	{
+		UE_LOG(LogMediaStream, Error, TEXT("Invalid Media Stream in FMediaStreamAssetSchemeHandler::CreateOrUpdateMediaPlayer"));
+		return nullptr;
+	}
+
+	const FMediaStreamSource& Source = InMediaStream->GetSource();
+
+	if (!IsValid(Source.Object))
+	{
+		UE_LOG(LogMediaStream, Error, TEXT("Invalid Asset in FMediaStreamAssetSchemeHandler::CreateOrUpdateMediaPlayer"));
+		return nullptr;
+	}
+
+	return FMediaStreamObjectHandlerManager::Get().CreateOrUpdateMediaPlayer(InMediaStream, Source.Object, InCurrentPlayer);
+}
+
+UObject* FMediaStreamAssetSchemeHandler::ResolveAsset(const FString& InPath)
+{
+	FSoftObjectPath SoftPath = InPath;
+
+	if (!UMediaStreamSourceBlueprintLibrary::IsAssetSoftPathValid(SoftPath))
+	{
+		UE_LOG(LogMediaStream, Error, TEXT("Invalid Asset Path in FMediaStreamAssetSchemeHandler::ResolveAsset [%s]"), *InPath);
+		return nullptr;
+	}
+
+	TSoftObjectPtr<UObject> SoftObjectPtr(SoftPath);
+	UObject* Asset = SoftObjectPtr.LoadSynchronous();
+
+	if (!IsValid(Asset))
+	{
+		UE_LOG(LogMediaStream, Error, TEXT("Invalid Asset in FMediaStreamAssetSchemeHandler::ResolveAsset [%s]"), *InPath);
+		return nullptr;
+	}
+
+	if (!FMediaStreamObjectHandlerManager::Get().CanHandleObject(Asset))
+	{
+		UE_LOG(LogMediaStream, Error, TEXT("Invalid Asset Class in FMediaStreamAssetSchemeHandler::CreateSource [%s]"), *InPath);
+		return nullptr;
+	}
+
+	return Asset;
+}
+
+#if WITH_EDITOR
+void FMediaStreamAssetSchemeHandler::CreatePropertyCustomization(TSharedRef<IPropertyHandle> InPropertyHandle, IDetailGroup& InDetailsGroup, IPropertyTypeCustomizationUtils& InCustomizationUtils)
+{
+	AddAssetSelector(InPropertyHandle, InDetailsGroup, InCustomizationUtils);
+}
+
+void FMediaStreamAssetSchemeHandler::AddAssetSelector(TSharedRef<IPropertyHandle> InPropertyHandle, IDetailGroup& InSourceGroup, IPropertyTypeCustomizationUtils& InCustomizationUtils)
+{
+	TWeakPtr<IPropertyHandle> PropertyHandleWeak = InPropertyHandle;
+
+	FMediaStreamSource* StreamSource = FMediaStreamSchemeHandlerLibrary::GetStreamSourcePtr(InPropertyHandle);
+
+	if (!StreamSource)
+	{
+		return;
+	}
+
+	TSharedPtr<IPropertyHandle> AssetProperty = InPropertyHandle->GetChildHandle(GET_MEMBER_NAME_CHECKED(FMediaStreamSource, Object), /* Recurse */ false);
+
+	if (!AssetProperty.IsValid())
+	{
+		return;
+	}
+
+	FDetailWidgetRow& FilePathRow = InSourceGroup.AddWidgetRow();
+
+	FilePathRow.Visibility(TAttribute<EVisibility>::CreateSP(this, &FMediaStreamAssetSchemeHandler::GetAssetSelectorVisibility, PropertyHandleWeak));
+
+	FilePathRow.NameContent()
+	[
+		SNew(STextBlock)
+		.Font(IDetailLayoutBuilder::GetDetailFont())
+		.Text(LOCTEXT("Asset", "Asset"))
+	];
+
+	FilePathRow.ValueContent()
+	[
+		SNew(SObjectPropertyEntryBox)
+		.ObjectPath_Static(&FMediaStreamSchemeHandlerLibrary::GetPath, PropertyHandleWeak)
+		.OnObjectChanged(this, &FMediaStreamAssetSchemeHandler::OnAssetSelected, PropertyHandleWeak)
+		.ThumbnailPool(InCustomizationUtils.GetThumbnailPool())
+		.DisplayThumbnail(true)
+		.AllowedClass(UMediaSource::StaticClass())
+	];
+}
+
+EVisibility FMediaStreamAssetSchemeHandler::GetAssetSelectorVisibility(TWeakPtr<IPropertyHandle> InPropertyHandleWeak) const
+{
+	if (FMediaStreamSchemeHandlerLibrary::GetScheme(InPropertyHandleWeak) == Scheme)
+	{
+		return EVisibility::Visible;
+	}
+
+	return EVisibility::Collapsed;
+}
+
+void FMediaStreamAssetSchemeHandler::OnAssetSelected(const FAssetData& InAssetData, TWeakPtr<IPropertyHandle> InPropertyHandleWeak)
+{
+	if (TSharedPtr<IPropertyHandle> PropertyHandle = InPropertyHandleWeak.Pin())
+	{
+		FMediaStreamSchemeHandlerLibrary::SetSource(PropertyHandle.ToSharedRef(), Scheme, InAssetData.GetSoftObjectPath().ToString());
+	}
+}
+#endif
+
+#undef LOCTEXT_NAMESPACE
diff --git a/Engine/Plugins/Experimental/MediaStream/Source/MediaStream/Private/SchemeHandlers/MediaStreamAssetSchemeHandler.h b/Engine/Plugins/Experimental/MediaStream/Source/MediaStream/Private/SchemeHandlers/MediaStreamAssetSchemeHandler.h
new file mode 100644
--- /dev/null
+++ b/Engine/Plugins/Experimental/MediaStream/Source/MediaStream/Private/SchemeHandlers/MediaStreamAssetSchemeHandler.h
@@ -0,0 +1,38 @@
+// Copyright Epic Games, Inc. All Rights Reserved.
+
+#pragma once
+
+#include "IMediaStreamSchemeHandler.h"
+
+class UMediaSource;
+
+#if WITH_EDITOR
+struct EVisibility;
+struct FAssetData;
+#endif
+
+class FMediaStreamAssetSchemeHandler : public IMediaStreamSchemeHandler
+{
+public:
+	static const FLazyName Scheme;
+
+	//~ Begin IMediaStreamSchemeHandler
+	virtual FMediaStreamSource CreateSource(UObject* InOuter, const FString& InPath) override;
+	virtual UMediaPlayer* CreateOrUpdateMediaPlayer(UMediaStream* InMediaStream, UMediaPlayer* InCurrentPlayer) override;
+
+#if WITH_EDITOR
+	virtual void CreatePropertyCustomization(TSharedRef<IPropertyHandle> InPropertyHandle, IDetailGroup& InDetailsGroup, IPropertyTypeCustomizationUtils& InCustomizationUtils) override;
+#endif
+	//~ End IMediaStreamSchemeHandler
+
+protected:
+	UObject* ResolveAsset(const FString& InPath);
+
+#if WITH_EDITOR
+	void AddAssetSelector(TSharedRef<IPropertyHandle> InPropertyHandle, IDetailGroup& InSourceGroup, IPropertyTypeCustomizationUtils& InCustomizationUtils);
+
+	EVisibility GetAssetSelectorVisibility(TWeakPtr<IPropertyHandle> InPropertyHandleWeak) const;
+
+	void OnAssetSelected(const FAssetData& InAssetData, TWeakPtr<IPropertyHandle> InPropertyHandleWeak);
+#endif
+};
diff --git a/Engine/Plugins/Experimental/MediaStream/Source/MediaStream/Private/SchemeHandlers/MediaStreamFileSchemeHandler.cpp b/Engine/Plugins/Experimental/MediaStream/Source/MediaStream/Private/SchemeHandlers/MediaStreamFileSchemeHandler.cpp
new file mode 100644
--- /dev/null
+++ b/Engine/Plugins/Experimental/MediaStream/Source/MediaStream/Private/SchemeHandlers/MediaStreamFileSchemeHandler.cpp
@@ -0,0 +1,184 @@
+// Copyright Epic Games, Inc. All Rights Reserved.
+
+#include "SchemeHandlers/MediaStreamFileSchemeHandler.h"
+
+#include "HAL/FileManager.h"
+#include "MediaSource.h"
+#include "MediaStream.h"
+#include "MediaStreamModule.h"
+#include "MediaStreamObjectHandlerManager.h"
+#include "MediaStreamSource.h"
+#include "Misc/Paths.h"
+
+#if WITH_EDITOR
+#include "DetailLayoutBuilder.h"
+#include "DetailWidgetRow.h"
+#include "IDetailGroup.h"
+#include "Widgets/Input/SFilePathPicker.h"
+#include "Widgets/Text/STextBlock.h"
+#endif
+
+#define LOCTEXT_NAMESPACE "MediaStreamFileSchemeHandler"
+
+const FLazyName FMediaStreamFileSchemeHandler::Scheme = TEXT("File");
+
+FMediaStreamSource FMediaStreamFileSchemeHandler::CreateSource(UObject* InOuter, const FString& InPath)
+{
+	FMediaStreamSource Source;
+
+	if (!IsValid(InOuter))
+	{
+		UE_LOG(LogMediaStream, Error, TEXT("Invalid Media Stream in FMediaStreamFileSchemeHandler::CreateSource"));
+		return Source;
+	}
+
+	Source.Scheme = Scheme;
+
+	UMediaSource* MediaSource = CreateMediaSource(InOuter, InPath);
+
+	if (!MediaSource)
+	{
+		return Source;
+	}
+
+	Source.Path = InPath;
+	Source.Object = MediaSource;
+
+	return Source;
+}
+
+UMediaPlayer* FMediaStreamFileSchemeHandler::CreateOrUpdateMediaPlayer(UMediaStream* InMediaStream, UMediaPlayer* InCurrentPlayer)
+{
+	if (!IsValid(InMediaStream))
+	{
+		UE_LOG(LogMediaStream, Error, TEXT("Invalid Media Stream in FMediaStreamFileSchemeHandler::CreateOrUpdateMediaPlayer"));
+		return nullptr;
+	}
+
+	UObject* MediaObject = InMediaStream->GetSource().Object;
+
+	if (!IsValid(MediaObject))
+	{
+		MediaObject = CreateMediaSource(InMediaStream, InMediaStream->GetSource().Path);
+
+		if (!IsValid(MediaObject))
+		{
+			return nullptr;
+		}
+	}
+
+	UMediaSource* MediaSource = Cast<UMediaSource>(MediaObject);
+
+	if (!IsValid(MediaSource))
+	{
+		UE_LOG(LogMediaStream, Error, TEXT("Invalid Media Path in FMediaStreamFileSchemeHandler::CreateOrUpdateMediaPlayer [%s]"), *MediaObject->GetClass()->GetName());
+		return nullptr;
+	}
+
+	return FMediaStreamObjectHandlerManager::Get().CreateOrUpdateMediaPlayer(InMediaStream, MediaSource, InCurrentPlayer);
+}
+
+FString FMediaStreamFileSchemeHandler::ResolveFilePath(const FString& InPath)
+{
+	FString FilePath = InPath;
+
+	if (IFileManager::Get().FileExists(*FilePath))
+	{
+		return FilePath;
+	}
+
+	const FString ProjectPath = FPaths::GetProjectFilePath();
+	FilePath = FPaths::Combine(ProjectPath, InPath);
+
+	if (IFileManager::Get().FileExists(*FilePath))
+	{
+		return FilePath;
+	}
+
+	const FString EnginePath = FPaths::EngineDir();
+	FilePath = FPaths::Combine(EnginePath, InPath);
+
+	if (IFileManager::Get().FileExists(*FilePath))
+	{
+		return FilePath;
+	}
+
+	UE_LOG(LogMediaStream, Error, TEXT("Invalid File Path in FMediaStreamFileSchemeHandler::ResolveFilePath [%s]"), *InPath);
+	return "";
+}
+
+UMediaSource* FMediaStreamFileSchemeHandler::CreateMediaSource(UObject* InOuter, const FString& InPath)
+{
+	const FString FilePath = ResolveFilePath(InPath);
+
+	if (FilePath.IsEmpty())
+	{
+		return nullptr;
+	}
+
+	UMediaSource* MediaSource = UMediaSource::SpawnMediaSourceForString(FilePath, InOuter);
+
+	if (!MediaSource)
+	{
+		UE_LOG(LogMediaStream, Error, TEXT("Invalid Media Path in FMediaStreamFileSchemeHandler::CreateMediaSource [%s]"), *FilePath);
+		return nullptr;
+	}
+
+	return MediaSource;
+}
+
+#if WITH_EDITOR
+void FMediaStreamFileSchemeHandler::CreatePropertyCustomization(TSharedRef<IPropertyHandle> InPropertyHandle, IDetailGroup& InDetailsGroup, IPropertyTypeCustomizationUtils& InCustomizationUtils)
+{
+	AddFileSelector(InPropertyHandle, InDetailsGroup, InCustomizationUtils);
+}
+
+void FMediaStreamFileSchemeHandler::AddFileSelector(TSharedRef<IPropertyHandle> InPropertyHandle, IDetailGroup& InSourceGroup, IPropertyTypeCustomizationUtils& InCustomizationUtils)
+{
+	TWeakPtr<IPropertyHandle> PropertyHandleWeak = InPropertyHandle;
+
+	FDetailWidgetRow& FilePathRow = InSourceGroup.AddWidgetRow();
+
+	FilePathRow.Visibility(TAttribute<EVisibility>::CreateSP(this, &FMediaStreamFileSchemeHandler::GetFileSelectorVisibility, PropertyHandleWeak));
+
+	FilePathRow.NameContent()
+	[
+		SNew(STextBlock)
+		.Font(IDetailLayoutBuilder::GetDetailFont())
+		.Text(LOCTEXT("FilePath", "File Path"))
+	];
+
+	FilePathRow.ValueContent()
+	[
+		SNew(SFilePathPicker)
+		.BrowseButtonImage(FAppStyle::GetBrush("PropertyWindow.Button_Ellipsis"))
+		.BrowseButtonStyle(FAppStyle::Get(), "HoverHintOnly")
+		.BrowseButtonToolTip(LOCTEXT("FileButtonToolTipText", "Choose a file from this computer"))
+		.BrowseTitle(LOCTEXT("PropertyEditorTitle", "File picker..."))
+		.FilePath_Static(&FMediaStreamSchemeHandlerLibrary::GetPath, PropertyHandleWeak)
+		.FileTypeFilter(TEXT("All files (*.*)|*.*"))
+		.OnPathPicked(this, &FMediaStreamFileSchemeHandler::OnFilePicked, PropertyHandleWeak)
+	];
+}
+
+EVisibility FMediaStreamFileSchemeHandler::GetFileSelectorVisibility(TWeakPtr<IPropertyHandle> InPropertyHandleWeak) const
+{
+	if (FMediaStreamSchemeHandlerLibrary::GetScheme(InPropertyHandleWeak) == Scheme)
+	{
+		return EVisibility::Visible;
+	}
+
+	return EVisibility::Collapsed;
+}
+
+void FMediaStreamFileSchemeHandler::OnFilePicked(const FString& InFilePath, TWeakPtr<IPropertyHandle> InPropertyHandleWeak)
+{
+	if (TSharedPtr<IPropertyHandle> PropertyHandle = InPropertyHandleWeak.Pin())
+	{
+		FMediaStreamSchemeHandlerLibrary::SetSource(PropertyHandle.ToSharedRef(), Scheme, InFilePath);
+	}
+}
+
+#endif
+
+#undef LOCTEXT_NAMESPACE
diff --git a/Engine/Plugins/Experimental/MediaStream/Source/MediaStream/Private/SchemeHandlers/MediaStreamFileSchemeHandler.h b/Engine/Plugins/Experimental/MediaStream/Source/MediaStream/Private/SchemeHandlers/MediaStreamFileSchemeHandler.h
new file mode 100644
--- /dev/null
+++ b/Engine/Plugins/Experimental/MediaStream/Source/MediaStream/Private/SchemeHandlers/MediaStreamFileSchemeHandler.h
@@ -0,0 +1,36 @@
+// Copyright Epic Games, Inc. All Rights Reserved.
+
+#pragma once
+
+#include "IMediaStreamSchemeHandler.h"
+
+class UMediaSource;
+struct EVisibility;
+
+class FMediaStreamFileSchemeHandler : public IMediaStreamSchemeHandler
+{
+public:
+	static const FLazyName Scheme;
+
+	//~ Begin IMediaStreamSchemeHandler
+	virtual FMediaStreamSource CreateSource(UObject* InOuter, const FString& InPath) override;
+	virtual UMediaPlayer* CreateOrUpdateMediaPlayer(UMediaStream* InMediaStream, UMediaPlayer* InCurrentPlayer) override;
+
+#if WITH_EDITOR
+	virtual void CreatePropertyCustomization(TSharedRef<IPropertyHandle> InPropertyHandle, IDetailGroup& InDetailsGroup, IPropertyTypeCustomizationUtils& InCustomizationUtils) override;
+#endif
+	//~ End IMediaStreamSchemeHandler
+
+protected:
+	FString ResolveFilePath(const FString& InPath);
+
+	UMediaSource* CreateMediaSource(UObject* InOuter, const FString& InPath);
+
+#if WITH_EDITOR
+	void AddFileSelector(TSharedRef<IPropertyHandle> InPropertyHandle, IDetailGroup& InSourceGroup, IPropertyTypeCustomizationUtils& InCustomizationUtils);
+
+	EVisibility GetFileSelectorVisibility(TWeakPtr<IPropertyHandle> InPropertyHandleWeak) const;
+
+	void OnFilePicked(const FString& InFilePath, TWeakPtr<IPropertyHandle> InPropertyHandleWeak);
+#endif
+};
diff --git a/Engine/Plugins/Experimental/MediaStream/Source/MediaStream/Private/SchemeHandlers/MediaStreamManagedSchemeHandler.cpp b/Engine/Plugins/Experimental/MediaStream/Source/MediaStream/Private/SchemeHandlers/MediaStreamManagedSchemeHandler.cpp
new file mode 100644
--- /dev/null
+++ b/Engine/Plugins/Experimental/MediaStream/Source/MediaStream/Private/SchemeHandlers/MediaStreamManagedSchemeHandler.cpp
@@ -0,0 +1,46 @@
+// Copyright Epic Games, Inc. All Rights Reserved.
+
+#include "SchemeHandlers/MediaStreamManagedSchemeHandler.h"
+
+#include "MediaStream.h"
+#include "MediaStreamModule.h"
+#include "MediaStreamSource.h"
+
+const FLazyName FMediaStreamManagedSchemeHandler::Scheme = TEXT("Managed");
+
+FMediaStreamSource FMediaStreamManagedSchemeHandler::CreateSource(UObject* InOuter, const FString& InPath)
+{
+	FMediaStreamSource Source;
+
+	if (!IsValid(InOuter))
+	{
+		UE_LOG(LogMediaStream, Error, TEXT("Invalid Media Stream in FMediaStreamManagedSchemeHandler::CreateSource"));
+		return Source;
+	}
+
+	Source.Scheme = Scheme;
+
+	// TODO: Check valid media source manager stream
+	if (false)
+	{
+		UE_LOG(LogMediaStream, Error, TEXT("Invalid Managed Stream in FMediaStreamManagedSchemeHandler::CreateSource [%s]"), *InPath);
+		return Source;
+	}
+
+	Source.Path = InPath;
+
+	return Source;
+}
+
+UMediaPlayer* FMediaStreamManagedSchemeHandler::CreateOrUpdateMediaPlayer(UMediaStream* InMediaStream, UMediaPlayer* InCurrentPlayer)
+{
+	// TODO: Get player from media source manager
+	return nullptr;
+}
+
+#if WITH_EDITOR
+void FMediaStreamManagedSchemeHandler::CreatePropertyCustomization(TSharedRef<IPropertyHandle> InPropertyHandle, IDetailGroup& InDetailsGroup, IPropertyTypeCustomizationUtils& InCustomizationUtils)
+{
+	// TODO Implement proper customisation
+}
+#endif
diff --git a/Engine/Plugins/Experimental/MediaStream/Source/MediaStream/Private/SchemeHandlers/MediaStreamManagedSchemeHandler.h b/Engine/Plugins/Experimental/MediaStream/Source/MediaStream/Private/SchemeHandlers/MediaStreamManagedSchemeHandler.h
new file mode 100644
--- /dev/null
+++ b/Engine/Plugins/Experimental/MediaStream/Source/MediaStream/Private/SchemeHandlers/MediaStreamManagedSchemeHandler.h
@@ -0,0 +1,20 @@
+// Copyright Epic Games, Inc. All Rights Reserved.
+
+#pragma once
+
+#include "IMediaStreamSchemeHandler.h"
+
+class FMediaStreamManagedSchemeHandler : public IMediaStreamSchemeHandler
+{
+public:
+	static const FLazyName Scheme;
+
+	//~ Begin IMediaStreamSchemeHandler
+	virtual FMediaStreamSource CreateSource(UObject* InOuter, const FString& InPath) override;
+	virtual UMediaPlayer* CreateOrUpdateMediaPlayer(UMediaStream* InMediaStream, UMediaPlayer* InCurrentPlayer) override;
+
+#if WITH_EDITOR
+	virtual void CreatePropertyCustomization(TSharedRef<IPropertyHandle> InPropertyHandle, IDetailGroup& InDetailsGroup, IPropertyTypeCustomizationUtils& InCustomizationUtils) override;
+#endif
+	//~ End IMediaStreamSchemeHandler
+};
diff --git a/Engine/Plugins/Experimental/MediaStream/Source/MediaStream/Private/SchemeHandlers/MediaStreamSubobjectSchemeHandler.cpp b/Engine/Plugins/Experimental/MediaStream/Source/MediaStream/Private/SchemeHandlers/MediaStreamSubobjectSchemeHandler.cpp
new file mode 100644
--- /dev/null
+++ b/Engine/Plugins/Experimental/MediaStream/Source/MediaStream/Private/SchemeHandlers/MediaStreamSubobjectSchemeHandler.cpp
@@ -0,0 +1,67 @@
+// Copyright Epic Games, Inc. All Rights Reserved.
+
+#include "SchemeHandlers/MediaStreamSubobjectSchemeHandler.h"
+
+#include "MediaStream.h"
+#include "MediaStreamModule.h"
+#include "MediaStreamObjectHandlerManager.h"
+#include "MediaStreamSource.h"
+
+const FLazyName FMediaStreamSubobjectSchemeHandler::Scheme = TEXT("Subobject");
+
+FMediaStreamSource FMediaStreamSubobjectSchemeHandler::CreateSource(UObject* InOuter, const FString& InPath)
+{
+	FMediaStreamSource Source;
+
+	if (!IsValid(InOuter))
+	{
+		UE_LOG(LogMediaStream, Error, TEXT("Invalid Media Stream in FMediaStreamSubobjectSchemeHandler::CreateSource"));
+		return Source;
+	}
+
+	Source.Scheme = Scheme;
+
+	UObject* Object = ResolveSubobjectPath(InOuter, InPath);
+
+	if (!Object)
+	{
+		UE_LOG(LogMediaStream, Error, TEXT("Invalid Subobject Path in FMediaStreamSubobjectSchemeHandler::CreateSource [%s]"), *InPath);
+		return Source;
+	}
+
+	Source.Path = InPath;
+
+	return Source;
+}
+
+UMediaPlayer* FMediaStreamSubobjectSchemeHandler::CreateOrUpdateMediaPlayer(UMediaStream* InMediaStream, UMediaPlayer* InCurrentPlayer)
+{
+	if (!IsValid(InMediaStream))
+	{
+		UE_LOG(LogMediaStream, Error, TEXT("Invalid Media Stream in FMediaStreamSubobjectSchemeHandler::CreateOrUpdateMediaPlayer"));
+		return nullptr;
+	}
+
+	UObject* Subobject = ResolveSubobjectPath(InMediaStream, InMediaStream->GetSource().Path);
+
+	if (!IsValid(Subobject))
+	{
+		UE_LOG(LogMediaStream, Error, TEXT("Invalid Subobject in FMediaStreamSubobjectSchemeHandler::CreateOrUpdateMediaPlayer"));
+		return nullptr;
+	}
+
+	return FMediaStreamObjectHandlerManager::Get().CreateOrUpdateMediaPlayer(InMediaStream, Subobject, InCurrentPlayer);
+}
+
+UObject* FMediaStreamSubobjectSchemeHandler::ResolveSubobjectPath(UObject* InOuter, const FString& InPath)
+{
+	const FSoftObjectPath SoftPath = InOuter->GetPathName() + TEXT(".") + InPath;
+	return SoftPath.ResolveObject();
+}
+
+#if WITH_EDITOR
+void FMediaStreamSubobjectSchemeHandler::CreatePropertyCustomization(TSharedRef<IPropertyHandle> InPropertyHandle, IDetailGroup& InDetailsGroup, IPropertyTypeCustomizationUtils& InCustomizationUtils)
+{
+	// TODO Implement proper customisation
+}
+#endif
diff --git a/Engine/Plugins/Experimental/MediaStream/Source/MediaStream/Private/SchemeHandlers/MediaStreamSubobjectSchemeHandler.h b/Engine/Plugins/Experimental/MediaStream/Source/MediaStream/Private/SchemeHandlers/MediaStreamSubobjectSchemeHandler.h
new file mode 100644
--- /dev/null
+++ b/Engine/Plugins/Experimental/MediaStream/Source/MediaStream/Private/SchemeHandlers/MediaStreamSubobjectSchemeHandler.h
@@ -0,0 +1,25 @@
+// Copyright Epic Games, Inc. All Rights Reserved.
+
+#pragma once
+
+#include "IMediaStreamSchemeHandler.h"
+
+struct FMediaStreamSource;
+
+class FMediaStreamSubobjectSchemeHandler : public IMediaStreamSchemeHandler
+{
+public:
+	static const FLazyName Scheme;
+
+	//~ Begin IMediaStreamSchemeHandler
+	virtual FMediaStreamSource CreateSource(UObject* InOuter, const FString& InPath) override;
+	virtual UMediaPlayer* CreateOrUpdateMediaPlayer(UMediaStream* InMediaStream, UMediaPlayer* InCurrentPlayer) override;
+
+#if WITH_EDITOR
+	virtual void CreatePropertyCustomization(TSharedRef<IPropertyHandle> InPropertyHandle, IDetailGroup& InDetailsGroup, IPropertyTypeCustomizationUtils& InCustomizationUtils) override;
+#endif
+	//~ End IMediaStreamSchemeHandler
+
+protected:
+	UObject* ResolveSubobjectPath(UObject* InOuter, const FString& InPath);
+};
diff --git a/Engine/Plugins/Experimental/MediaStream/Source/MediaStream/Public/IMediaStreamObjectHandler.h b/Engine/Plugins/Experimental/MediaStream/Source/MediaStream/Public/IMediaStreamObjectHandler.h
new file mode 100644
--- /dev/null
+++ b/Engine/Plugins/Experimental/MediaStream/Source/MediaStream/Public/IMediaStreamObjectHandler.h
@@ -0,0 +1,28 @@
+// Copyright Epic Games, Inc. All Rights Reserved.
+
+#pragma once
+
+#include "Templates/SharedPointer.h"
+
+class UMediaPlayer;
+class UMediaStream;
+class UObject;
+
+/**
+ * Implement this interface to add a new object handler.
+ * If the derived class has a static UClass* Class then it can be added without specifying it.
+ */
+class IMediaStreamObjectHandler : public TSharedFromThis<IMediaStreamObjectHandler>
+{
+public:
+	virtual ~IMediaStreamObjectHandler() = default;
+
+	/**
+	 * Create a media player for the provided object.
+	 * @param InMediaStream The container for the player.
+	 * @param InObject The object for which to create a player.
+	 * @param InCurrentPlayer The current player to update or null.
+	 * @return A player or nullptr.
+	 */
+	virtual UMediaPlayer* CreateOrUpdatePlayer(UMediaStream* InMediaStream, UObject* InObject, UMediaPlayer* InCurrentPlayer) = 0;
+};
diff --git a/Engine/Plugins/Experimental/MediaStream/Source/MediaStream/Public/IMediaStreamPlayer.h b/Engine/Plugins/Experimental/MediaStream/Source/MediaStream/Public/IMediaStreamPlayer.h
new file mode 100644
--- /dev/null
+++ b/Engine/Plugins/Experimental/MediaStream/Source/MediaStream/Public/IMediaStreamPlayer.h
@@ -0,0 +1,217 @@
+// Copyright Epic Games, Inc. All Rights Reserved.
+
+#pragma once
+
+#include "UObject/Interface.h"
+
+#include "MediaStreamEnums.h"
+#include "MediaStreamPlayerConfig.h"
+#include "MediaStreamTextureConfig.h"
+
+#include "IMediaStreamPlayer.generated.h"
+
+class ULevelSequence;
+class UMediaPlayer;
+class UMediaStream;
+class UMediaTexture;
+struct FGuid;
+struct FMediaStreamSource;
+
+/**
+ * Interface for Media Stream Players. Should only be used within a Media Stream object.
+ */
+UINTERFACE(MinimalAPI, BlueprintType, Category = "Media Stream", meta = (CannotImplementInterfaceInBlueprint))
+class UMediaStreamPlayer : public UInterface
+{
+	GENERATED_BODY()
+};
+
+/**
+ * Interface for Media Stream Players. Should only be used within a Media Stream object.
+ */
+class IMediaStreamPlayer
+{
+	GENERATED_BODY()
+
+public:
+	UFUNCTION(BlueprintCallable, Category = "Media Stream|Player")
+	virtual UMediaStream* GetMediaStream() const = 0;
+
+	/* Returns whether this player's controls can do anything. */
+	UFUNCTION(BlueprintCallable, CallInEditor, Category = "Media Stream|Player")
+	virtual bool IsReadOnly() const = 0;
+
+	UFUNCTION(BlueprintCallable, Category = "Media Stream|Player")
+	virtual void OnCreated() = 0;
+
+	/**
+	 * Called by the owning Media Stream when the source changes so that the player can update.
+	 */
+	virtual void OnSourceChanged(const FMediaStreamSource& InSource) = 0;
+
+	/**
+	 * @return Gets the media texture subobject.
+	 */
+	UFUNCTION(BlueprintCallable, Category = "Media Stream|Texture")
+	virtual UMediaTexture* GetMediaTexture() const = 0;
+
+	/**
+	 * @return The current texture config.
+	 */
+	UFUNCTION(BlueprintCallable, Category = "Media Stream|Texture")
+	virtual const FMediaStreamTextureConfig& GetTextureConfig() const = 0;
+
+	/**
+	 * Update the current texture's (and any newly set texture's) config.
+	 * @param InTextureConfig The new config.
+	 */
+	UFUNCTION(BlueprintCallable, Category = "Media Stream|Texture")
+	virtual void SetTextureConfig(const FMediaStreamTextureConfig& InTextureConfig) = 0;
+
+	/**
+	 * Applies the current texture config to the current media texture.
+	 * This is generally automatic.
+	 */
+	UFUNCTION(BlueprintCallable, Category = "Media Stream|Texture")
+	virtual void ApplyTextureConfig() = 0;
+
+	/**
+	 * @returns The active player, if there is one and it is valid.
+	 */
+	UFUNCTION(BlueprintCallable, Category = "Media Stream|Player")
+	virtual UMediaPlayer* GetPlayer() const = 0;
+
+	/**
+	 * @return True if there is an active and valid player.
+	 */
+	UFUNCTION(BlueprintCallable, Category = "Media Stream|Player")
+	virtual bool HasValidPlayer() const = 0;
+
+	/**
+	 * @return The current player config.
+	 */
+	UFUNCTION(BlueprintCallable, Category = "Media Stream|Player")
+	virtual const FMediaStreamPlayerConfig& GetPlayerConfig() const = 0;
+
+	/**
+	 * Update the current player's (and any newly set player's) config.
+	 * @param InPlayerConfig The new config.
+	 */
+	UFUNCTION(BlueprintCallable, Category = "Media Stream|Player")
+	virtual void SetPlayerConfig(const FMediaStreamPlayerConfig& InPlayerConfig) = 0;
+
+	/**
+	 * Applies the current player config to the current media player.
+	 * This is generally automatic.
+	 */
+	UFUNCTION(BlueprintCallable, Category = "Media Stream|Player")
+	virtual void ApplyPlayerConfig() = 0;
+
+	/**
+	 * Setter for events.
+	 * @param InIndex The index to start playing.
+	 * @return True on success. False on failure or invalid player.
+	 */
+	UFUNCTION(BlueprintCallable, Category = "Media Stream|Player")
+	virtual bool SetPlaylistIndex(int32 InIndex) = 0;
+
+	/**
+	 * @return The time offset in the media that has been requested.
+	 */
+	UFUNCTION(BlueprintCallable, Category = "Media Stream|Player")
+	virtual float GetRequestedSeekTime() const = 0;
+
+	/**
+	 * Sets the time in the currently playing player, if possible.
+	 * @param InTime The time offset in the media to request.
+	 * @return True on success. False on failure or invalid player.
+	 */
+	UFUNCTION(BlueprintCallable, Category = "Media Stream|Player")
+	virtual bool SetRequestedSeekTime(float InTime) = 0;
+
+	/**
+	 * @return Gets the media player playback state.
+	 */
+	UFUNCTION(BlueprintCallable, Category = "Media Stream|Player")
+	virtual EMediaStreamPlaybackState GetPlaybackState() const = 0;
+
+	/**
+	 * Sets the media player playback state.
+	 * @param InState The state to change to.
+	 * @return True on success. False on failure or invalid player.
+	 */
+	UFUNCTION(BlueprintCallable, Category = "Media Stream|Player")
+	virtual bool SetPlaybackState(EMediaStreamPlaybackState InState) = 0;
+
+	/**
+	 * @return Gets the active playlist index. May not match the actual player. -1 on error.
+	 */
+	UFUNCTION(BlueprintCallable, Category = "Media Stream|Player")
+	virtual int32 GetPlaylistIndex() const = 0;
+
+	/**
+	 * @return The size of the current playlist. -1 on error.
+	 */
+	UFUNCTION(BlueprintCallable, Category = "Media Stream|Player")
+	virtual int32 GetPlaylistNum() const = 0;
+
+	/**
+	 * @return For proxy streams, returns the stream this is getting data from.
+	 */
+	UFUNCTION(BlueprintCallable, Category = "Media Stream|Player")
+	virtual UMediaStream* GetSourceStream() const = 0;
+
+	/**
+	 * Continues play of the current media.
+	 * @return True if successful.
+	 */
+	UFUNCTION(BlueprintCallable, CallInEditor, Category = "Media Stream|Player")
+	virtual bool Play() = 0;
+
+	/**
+	 * Pauses playback of the current media.
+	 * @return True if successful.
+	 */
+	UFUNCTION(BlueprintCallable, CallInEditor, Category = "Media Stream|Player")
+	virtual bool Pause() = 0;
+
+	/**
+	 * Seeks to the start of the current media.
+	 * @return True if successful.
+	 */
+	UFUNCTION(BlueprintCallable, CallInEditor, Category = "Media Stream|Player")
+	virtual bool Rewind() = 0;
+
+	/**
+	 * Seeks to the end of the current media.
+	 * @return True if successful.
+	 */
+	UFUNCTION(BlueprintCallable, CallInEditor, Category = "Media Stream|Player")
+	virtual bool FastForward() = 0;
+
+	/**
+	 * Goes to the previous item in the playlist. Only possible with playlist sources.
+	 * @return True if successful.
+	 */
+	UFUNCTION(BlueprintCallable, CallInEditor, Category = "Media Stream|Player")
+	virtual bool Previous() = 0;
+
+	/**
+	 * Goes to the next item in the playlist. Only possible with playlist sources.
+	 * @return True if successful.
+	 */
+	UFUNCTION(BlueprintCallable, CallInEditor, Category = "Media Stream|Player")
+	virtual bool Next() = 0;
+
+	/**
+	 * Closes the current media player.
+	 * @return True if successful.
+	 */
+	UFUNCTION(BlueprintCallable, CallInEditor, Category = "Media Stream|Player")
+	virtual bool Close() = 0;
+
+	/**
+	 * Called by the Media Stream when it is done with this Media Stream Player.
+	 */
+	virtual void Deinitialize() = 0;
+};
diff --git a/Engine/Plugins/Experimental/MediaStream/Source/MediaStream/Public/IMediaStreamSchemeHandler.h b/Engine/Plugins/Experimental/MediaStream/Source/MediaStream/Public/IMediaStreamSchemeHandler.h
new file mode 100644
--- /dev/null
+++ b/Engine/Plugins/Experimental/MediaStream/Source/MediaStream/Public/IMediaStreamSchemeHandler.h
@@ -0,0 +1,61 @@
+// Copyright Epic Games, Inc. All Rights Reserved.
+
+#pragma once
+
+#include "Templates/SharedPointer.h"
+
+#include "Containers/ContainersFwd.h"
+
+class IPropertyHandle;
+class UMediaPlayer;
+class UMediaStream;
+struct FMediaStreamSource;
+
+#if WITH_EDITOR
+class IDetailGroup;
+class IPropertyTypeCustomizationUtils;
+#endif
+
+/**
+ * Implement this interface to add a new scheme handler.
+ * If the derived class has a static TCHAR Scheme[] then it can be added without specifying it.
+ */
+class IMediaStreamSchemeHandler : public TSharedFromThis<IMediaStreamSchemeHandler>
+{
+public:
+	virtual ~IMediaStreamSchemeHandler() = default;
+
+	/**
+	 * Passes a path to produce a stream source.
+	 * @param InOuter The container of the source.
+	 * @param InPath The path representing the source.
+	 * @return A valid stream source object or none if it was not a valid path.
+	 */
+	virtual FMediaStreamSource CreateSource(UObject* InOuter, const FString& InPath) = 0;
+
+	/**
+	 * Creates a media player for the given Media Stream based on its source.
+	 * @param InMediaStream The Media Stream to create a player for.
+	 * @param InCurrentPlayer The current player to update or null.
+	 * @return A media player or nullptr on failure.
+	 */
+	virtual UMediaPlayer* CreateOrUpdateMediaPlayer(UMediaStream* InMediaStream, UMediaPlayer* InCurrentPlayer) = 0;
+
+#if WITH_EDITOR
+	virtual void CreatePropertyCustomization(TSharedRef<IPropertyHandle> InPropertyHandle, IDetailGroup& InDetailsGroup, IPropertyTypeCustomizationUtils& InCustomizationUtils) = 0;
+#endif
+};
+
+#if WITH_EDITOR
+class FMediaStreamSchemeHandlerLibrary
+{
+public:
+	MEDIASTREAM_API static FMediaStreamSource* GetStreamSourcePtr(TWeakPtr<IPropertyHandle> InPropertyHandleWeak);
+
+	MEDIASTREAM_API static FName GetScheme(TWeakPtr<IPropertyHandle> InPropertyHandleWeak);
+
+	MEDIASTREAM_API static FString GetPath(TWeakPtr<IPropertyHandle> InPropertyHandleWeak);
+
+	MEDIASTREAM_API static void SetSource(TSharedRef<IPropertyHandle> InPropertyHandle, const FName& InScheme, const FString& InPath);
+};
+#endif
diff --git a/Engine/Plugins/Experimental/MediaStream/Source/MediaStream/Public/MediaStream.h b/Engine/Plugins/Experimental/MediaStream/Source/MediaStream/Public/MediaStream.h
new file mode 100644
--- /dev/null
+++ b/Engine/Plugins/Experimental/MediaStream/Source/MediaStream/Public/MediaStream.h
@@ -0,0 +1,110 @@
+// Copyright Epic Games, Inc. All Rights Reserved.
+
+#pragma once
+
+#include "MediaPlayerProxyInterface.h"
+#include "UObject/Object.h"
+
+#include "MediaStreamSource.h"
+#include "UObject/ScriptInterface.h"
+
+#include "MediaStream.generated.h"
+
+class IMediaStreamPlayer;
+
+/**
+ * Media Stream. Provides an agnostic interface between controllers and players.
+ */
+UCLASS(MinimalAPI, BlueprintType, Blueprintable, ClassGroup = "Media Stream", DefaultToInstanced)
+class UMediaStream : public UObject, public IMediaPlayerProxyInterface
+{
+	GENERATED_BODY()
+
+public:
+	virtual ~UMediaStream() override = default;
+
+	/**
+	 * @return True if the media source scheme is not empty. Does not guarantee the source is correct, only valid.
+	 */
+	UFUNCTION(BlueprintPure, Category = "Media Stream")
+	MEDIASTREAM_API bool HasValidSource() const;
+
+	/**
+	 * @return The current Media Stream Source. Does not guaranteed that it is valid.
+	 */
+	UFUNCTION(BlueprintPure, Category = "Media Stream")
+	MEDIASTREAM_API const FMediaStreamSource& GetSource() const;
+
+	/**
+	 * @return The source at the end of the chain of proxy players, or the local one if it is local. If
+	 * the chain is interrupted, the last link in the chain's source will be returned.
+	 */
+	UFUNCTION(BlueprintPure, Category = "Media Stream")
+	MEDIASTREAM_API const FMediaStreamSource& ResolveSource() const;
+
+	/**
+	 * Sets a new Media Stream Source and inits the Player. It should be ready to use if this returns true.
+	 * @param InSource The new source.
+	 * @return True if the source was successfully applied.
+	 */
+	UFUNCTION(BlueprintCallable, Category = "Media Stream")
+	MEDIASTREAM_API bool SetSource(const FMediaStreamSource& InSource);
+
+	/**
+	 * @return Gets the active Media Stream Player.
+	 */
+	UFUNCTION(BlueprintPure, Category = "Media Stream")
+	MEDIASTREAM_API TScriptInterface<IMediaStreamPlayer> GetPlayer() const;
+
+	/**
+	 * Called when the Media Stream Player changes.
+	 */
+	DECLARE_DYNAMIC_MULTICAST_DELEGATE_TwoParams(FOnPlayerChanged, UMediaStream*, InMediaStream, TScriptInterface<IMediaStreamPlayer>, InPlayer);
+	MEDIASTREAM_API FOnPlayerChanged& GetOnPlayerChanged();
+
+	//~ Begin UObject
+#if WITH_EDITOR
+	virtual void PreEditUndo() override;
+	virtual void PostEditUndo() override;
+	virtual void PostEditChangeProperty(FPropertyChangedEvent& InPropertyChangedEvent) override;
+#endif
+	virtual void PostDuplicate(bool bInDuplicateForPIE) override;
+	virtual void PostEditImport() override;
+	virtual void PostNetReceive() override;
+	virtual void BeginDestroy() override;
+	//~ End UObject
+
+	//~ Begin IMediaPlayerProxyInterface
+	virtual float GetProxyRate() const override;
+	virtual bool SetProxyRate(float Rate) override;
+	virtual bool IsExternalControlAllowed() override;
+	virtual const FMediaSourceCacheSettings& GetCacheSettings() const override;
+	virtual UMediaSource* ProxyGetMediaSourceFromIndex(int32 Index) const override;
+	virtual UMediaTexture* ProxyGetMediaTexture(int32 LayerIndex, int32 TextureIndex) override;
+	virtual void ProxyReleaseMediaTexture(int32 LayerIndex, int32 TextureIndex) override;
+	virtual bool ProxySetAspectRatio(UMediaPlayer* InMediaPlayer) override;
+	virtual void ProxySetTextureBlend(int32 LayerIndex, int32 TextureIndex, float Blend) override;
+	//~ End IMediaPlayerProxyInterface
+
+protected:
+	UPROPERTY(EditAnywhere, BlueprintReadWrite, Category = "Media Stream")
+	FMediaStreamSource Source;
+
+	/**
+	 * The player instance is automatically created based on the type of source.
+	 * Where possible, the player is re-used when opening new media.
+	 */
+	UPROPERTY(VisibleAnywhere, BlueprintReadOnly, Instanced, DisplayName = "Player", Category = "Media Stream")
+	TObjectPtr<UObject> PlayerObject = nullptr;
+
+	UPROPERTY(BlueprintAssignable)
+	FOnPlayerChanged OnPlayerChanged;
+
+	FMediaStreamSource Source_PreUndo;
+
+	/**
+	 * Applies the current source settings to this Media Stream, potentially changing the Media Stream Player.
+	 * @return True if it was successful.
+	 */
+	bool ApplySource();
+};
diff --git a/Engine/Plugins/Experimental/MediaStream/Source/MediaStream/Public/MediaStreamActor.h b/Engine/Plugins/Experimental/MediaStream/Source/MediaStream/Public/MediaStreamActor.h
new file mode 100644
--- /dev/null
+++ b/Engine/Plugins/Experimental/MediaStream/Source/MediaStream/Public/MediaStreamActor.h
@@ -0,0 +1,22 @@
+// Copyright Epic Games, Inc. All Rights Reserved.
+
+#pragma once
+
+#include "Components/StaticMeshComponent.h"
+
+#include "MediaStreamActor.generated.h"
+
+class UMediaStreamComponent;
+
+UCLASS(BlueprintType, MinimalAPI, ComponentWrapperClass)
+class AMediaStreamActor : public AActor
+{
+	GENERATED_BODY()
+
+public:
+	AMediaStreamActor();
+
+protected:
+	UPROPERTY(VisibleAnywhere, BlueprintReadOnly, Category = "Media Stream")
+	TObjectPtr<UMediaStreamComponent> MediaStreamComponent;
+};
diff --git a/Engine/Plugins/Experimental/MediaStream/Source/MediaStream/Public/MediaStreamComponent.h b/Engine/Plugins/Experimental/MediaStream/Source/MediaStream/Public/MediaStreamComponent.h
new file mode 100644
--- /dev/null
+++ b/Engine/Plugins/Experimental/MediaStream/Source/MediaStream/Public/MediaStreamComponent.h
@@ -0,0 +1,37 @@
+// Copyright Epic Games, Inc. All Rights Reserved.
+
+#pragma once
+
+#include "Components/StaticMeshComponent.h"
+
+#include "MediaStreamComponent.generated.h"
+
+class IMediaStreamPlayer;
+class UMediaStream;
+
+UCLASS(BlueprintType, MinimalAPI)
+class UMediaStreamComponent : public UStaticMeshComponent
+{
+	GENERATED_BODY()
+
+public:
+	UMediaStreamComponent();
+
+	//~ Begin UObject
+	virtual void PostInitProperties() override;
+	virtual void PostLoad() override;
+	virtual void PostNetReceive() override;
+	//~ End UObject
+
+protected:
+	UPROPERTY(VisibleAnywhere, BlueprintReadOnly, Category = "Media Stream")
+	TObjectPtr<UMediaStream> MediaStream;
+
+	UPROPERTY(BlueprintReadOnly, Transient, DuplicateTransient, TextExportTransient, Category = "Media Stream")
+	TObjectPtr<UMaterialInstanceDynamic> MaterialInstance;
+
+	UFUNCTION()
+	void OnPlayerChanged(UMediaStream* InMediaStream, TScriptInterface<IMediaStreamPlayer> InPlayer);
+
+	void InitPlayer();
+};
diff --git a/Engine/Plugins/Experimental/MediaStream/Source/MediaStream/Public/MediaStreamEnums.h b/Engine/Plugins/Experimental/MediaStream/Source/MediaStream/Public/MediaStreamEnums.h
new file mode 100644
--- /dev/null
+++ b/Engine/Plugins/Experimental/MediaStream/Source/MediaStream/Public/MediaStreamEnums.h
@@ -0,0 +1,30 @@
+// Copyright Epic Games, Inc. All Rights Reserved.
+
+#pragma once
+
+#include "HAL/Platform.h"
+
+#include "MediaStreamEnums.generated.h"
+
+UENUM(BlueprintType)
+enum class EMediaStreamPlaybackState : uint8
+{
+	Play,
+	Pause
+};
+
+UENUM(BlueprintType)
+enum class EMediaStreamPlaybackDirection : uint8
+{
+	Forward,
+	Backward
+};
+
+UENUM(BlueprintType)
+enum class EMediaStreamPlaybackSeek : uint8
+{
+	Previous,
+	Start,
+	End,
+	Next
+};
diff --git a/Engine/Plugins/Experimental/MediaStream/Source/MediaStream/Public/MediaStreamModule.h b/Engine/Plugins/Experimental/MediaStream/Source/MediaStream/Public/MediaStreamModule.h
new file mode 100644
--- /dev/null
+++ b/Engine/Plugins/Experimental/MediaStream/Source/MediaStream/Public/MediaStreamModule.h
@@ -0,0 +1,34 @@
+// Copyright Epic Games, Inc. All Rights Reserved.
+
+#pragma once
+
+#include "Modules/ModuleInterface.h"
+
+#include "Delegates/Delegate.h"
+#include "Delegates/DelegateCombinations.h"
+#include "Logging/LogMacros.h"
+
+enum class EMediaStreamPlaybackState : uint8;
+
+DECLARE_LOG_CATEGORY_EXTERN(LogMediaStream, Log, All);
+
+DECLARE_MULTICAST_DELEGATE_OneParam(FMediaStreamEditorPlaybackStateDelegate, EMediaStreamPlaybackState)
+
+/**
+ * Media Stream - Content/type agnostic chainable media proxy with media player integration.
+ */
+class FMediaStreamModule : public IModuleInterface
+{
+public:
+	MEDIASTREAM_API static FMediaStreamModule& Get();
+
+	MEDIASTREAM_API FMediaStreamEditorPlaybackStateDelegate& GetEditorPlaybackStateEvent();
+
+	//~ Begin IModuleInterface
+	virtual void StartupModule() override;
+	virtual void ShutdownModule() override;
+	//~ End IModuleInterface
+
+private:
+	FMediaStreamEditorPlaybackStateDelegate EditorPlaybackStateDelegate;
+};
diff --git a/Engine/Plugins/Experimental/MediaStream/Source/MediaStream/Public/MediaStreamObjectHandlerManager.h b/Engine/Plugins/Experimental/MediaStream/Source/MediaStream/Public/MediaStreamObjectHandlerManager.h
new file mode 100644
--- /dev/null
+++ b/Engine/Plugins/Experimental/MediaStream/Source/MediaStream/Public/MediaStreamObjectHandlerManager.h
@@ -0,0 +1,197 @@
+// Copyright Epic Games, Inc. All Rights Reserved.
+
+#pragma once
+
+#include "Containers/Map.h"
+#include "HAL/Platform.h"
+#include "Templates/SharedPointer.h"
+#include "UObject/NameTypes.h"
+
+class IMediaStreamObjectHandler;
+class UMediaPlayer;
+class UMediaStream;
+struct FMediaStreamSource;
+
+/**
+ * Handles the registration and operation of object handlers.
+ * Will fall back to using the Media Assets module if no handler is found.
+ */
+class FMediaStreamObjectHandlerManager
+{
+	friend class UMediaStreamSourceBlueprintLibrary;
+
+public:
+	/**
+	 * @return Gets the instance of this manager.
+	 */
+	MEDIASTREAM_API static FMediaStreamObjectHandlerManager& Get();
+
+	/**
+	 * Checks whether the given class has a class handler. Checks each super class too.
+	 * @param InClass The class to check.
+	 * @return True if any handler can handle the class.
+	 */
+	MEDIASTREAM_API bool CanHandleObject(const UClass* InClass) const;
+
+	/**
+	 * Templated version of CanHandleObject.
+	 * @tparam InClassName The class to check.
+	 * @return True if any handler can handle the class.
+	 */
+	template<typename InClassName
+		UE_REQUIRES(TModels_V<CStaticClassProvider, InClassName>)>
+	bool CanHandleObject() const
+	{
+		return CanHandleObject(InClassName::StaticClass());
+	}
+
+	/**
+	 * Checks whether the given object's class has a class handler. Checks each super class too.
+	 * @param InObject The object whose class is to be checked.
+	 * @return True if any handler can handle the object's class.
+	 */
+	MEDIASTREAM_API bool CanHandleObject(const UObject* InObject) const;
+
+	/**
+	 * Checks whether a class has a handler registered.
+	 * @param InClass The class to check. Must be an exact match.
+	 * @return True if the handler is registered.
+	 */
+	MEDIASTREAM_API bool HasObjectHandler(const UClass* InClass) const;
+
+	/**
+	 * Checks whether an object's class has a handler registered.
+	 * @param InObject The object whose class to check. Must be an exact match.
+	 * @return True if the handler is registered.
+	 */
+	MEDIASTREAM_API bool HasObjectHandler(const UObject* InObject) const;
+
+	/**
+	 * Templated version of HasObjectHandler.
+	 * @tparam InClassName The class to check. Must be an exact match.
+	 * @return True if the handler is registered.
+	 */
+	template<typename InClassName
+		UE_REQUIRES(TModels_V<CStaticClassProvider, InClassName>)>
+	bool HasObjectHandler() const
+	{
+		return HasObjectHandler(InClassName::StaticClass());
+	}
+
+	/**
+	 * Does an hierarchical search for a handler that can handle the given class.
+	 * @param InClass A validity checked UObject class.
+	 * @return The handler for the class, if found.
+	 */
+	MEDIASTREAM_API TSharedPtr<IMediaStreamObjectHandler> GetHandler(const UClass* InClass) const;
+
+	/**
+	 * Registers a Object handler that creates a media player for the given object class. Holds a strong reference.
+	 * Does not replace an already registered handler.
+	 * @param InClass The class of the object handled.
+	 * @param InHandler The handle which creates the media player.
+	 * @return True if the handler was registered.
+	 */
+	MEDIASTREAM_API bool RegisterObjectHandler(const UClass* InClass, const TSharedRef<IMediaStreamObjectHandler>& InHandler);
+
+	/**
+	 * Templated version of RegisterObjectHandler.
+	 * @tparam InClassName The class of the object handled.
+	 * @param InHandler The handle which creates the media player.
+	 * @return True if the handler was registered.
+	 */
+	template<typename InClassName
+		UE_REQUIRES(TModels_V<CStaticClassProvider, InClassName>)>
+	bool RegisterObjectHandler(const TSharedRef<IMediaStreamObjectHandler>& InHandler)
+	{
+		return RegisterObjectHandler(InClassName::StaticClass(), InHandler);
+	}
+
+	/**
+	 * Templated version of RegisterObjectHandler.
+	 * @tparam InHandlerClass Class of the handle which creates the media player.
+	 * @tparam InArgsType The type of handle class constructor parameters (automatic).
+	 * @param InClass The class of the object handled.
+	 * @param InArgs The handle class constructor parameters.
+	 * @return True if the handler was registered.
+	 */
+	template<typename InHandlerClass, typename... InArgsType
+		UE_REQUIRES(TIsDerivedFrom<InHandlerClass, IMediaStreamObjectHandler>::Value)>
+	bool RegisterObjectHandler(const UClass* InClass, InArgsType&&... InArgs)
+	{
+		return RegisterObjectHandler(InClass, MakeShared<InHandlerClass>(Forward<InArgsType>(InArgs)...));
+	}
+
+	/**
+	 * Templated version of RegisterObjectHandler.
+	 * @tparam InClassName The class of the object handled.
+	 * @tparam InHandlerClass Class of the handle which creates the media player.
+	 * @tparam InArgsType The type of handle class constructor parameters (automatic).
+	 * @param InArgs The handle class constructor parameters.
+	 * @return True if the handler was registered.
+	 */
+	template<typename InClassName, typename InHandlerClass, typename... InArgsType
+		UE_REQUIRES(TAnd<TModels<CStaticClassProvider, InClassName>,
+			TIsDerivedFrom<InHandlerClass, IMediaStreamObjectHandler>>::Value)>
+		bool RegisterObjectHandler(InArgsType&&... InArgs)
+	{
+		return RegisterObjectHandler(InClassName::StaticClass(), MakeShared<InHandlerClass>(Forward<InArgsType>(InArgs)...));
+	}
+
+	/**
+	 * Templated version of RegisterObjectHandler with auto Class value.
+	 * @tparam InHandlerClass Class of the handle which creates the media player.
+	 * @tparam InArgsType The type of handle class constructor parameters (automatic).
+	 * @param InArgs The handle class constructor parameters.
+	 * @return True if the handler was registered.
+	 */
+	template<typename InHandlerClass, typename... InArgsType
+		UE_REQUIRES(TIsDerivedFrom<InHandlerClass, IMediaStreamObjectHandler>::Value)>
+	bool RegisterObjectHandler(InArgsType&&... InArgs)
+	{
+		return RegisterObjectHandler(InHandlerClass::GetClass(), MakeShared<InHandlerClass>(Forward<InArgsType>(InArgs)...));
+	}
+
+	/**
+	 * Unregisters a Object handler.
+	 * @param InClass The class to unregister.
+	 * @return The removed Object handler, if found.
+	 */
+	MEDIASTREAM_API TSharedPtr<IMediaStreamObjectHandler> UnregisterObjectHandler(const UClass* InClass);
+
+	/**
+	 * Templated version of UnregisterObjectHandler.
+	 * @tparam InClassName The class to unregister.
+	 * @return The removed Object handler, if found.
+	 */
+	template<typename InClassName
+		UE_REQUIRES(TModels_V<CStaticClassProvider, InClassName>)>
+	TSharedPtr<IMediaStreamObjectHandler> UnregisterObjectHandler()
+	{
+		return RegisterObjectHandler(InClassName::StaticClass());
+	}
+
+	/**
+	 * Templated version of UnregisterObjectHandler with auto Class value.
+	 * @tparam InHandlerClass Class of the handle which creates the media player.
+	 * @return The removed Object handler, if found.
+	 */
+	template<typename InHandlerClass
+		UE_REQUIRES(TIsDerivedFrom<InHandlerClass, IMediaStreamObjectHandler>::Value)>
+	TSharedPtr<InHandlerClass> UnregisterObjectHandler()
+	{
+		return StaticCastSharedPtr<InHandlerClass>(UnregisterObjectHandler(InHandlerClass::GetClass()));
+	}
+
+	/**
+	 * Creates a media player for a given Media Stream and source object.
+	 * @param InMediaStream The Media Stream container.
+	 * @param InSourceObject The playback source.
+	 * @param InCurrentPlayer The current player to update or null.
+	 * @return A media player or nullptr on failure.
+	 */
+	MEDIASTREAM_API UMediaPlayer* CreateOrUpdateMediaPlayer(UMediaStream* InMediaStream, UObject* InSourceObject, UMediaPlayer* InCurrentPlayer) const;
+
+private:
+	TMap<FName, TSharedRef<IMediaStreamObjectHandler>> Handlers;
+};
diff --git a/Engine/Plugins/Experimental/MediaStream/Source/MediaStream/Public/MediaStreamObjectHandlerSubsystem.h b/Engine/Plugins/Experimental/MediaStream/Source/MediaStream/Public/MediaStreamObjectHandlerSubsystem.h
new file mode 100644
--- /dev/null
+++ b/Engine/Plugins/Experimental/MediaStream/Source/MediaStream/Public/MediaStreamObjectHandlerSubsystem.h
@@ -0,0 +1,52 @@
+// Copyright Epic Games, Inc. All Rights Reserved.
+
+#pragma once
+
+#include "Subsystems/EngineSubsystem.h"
+
+#include "MediaStreamObjectHandlerSubsystem.generated.h"
+
+class UMediaPlayer;
+class UMediaStream;
+
+/**
+ * Subsystem for blueprint-style interaction with the Object Handler Manager.
+ */
+UCLASS(MinimalAPI, BlueprintType, ClassGroup = "Media Stream")
+class UMediaStreamObjectHandlerSubsystem : public UEngineSubsystem
+{
+	GENERATED_BODY()
+
+public:
+	/**
+	 * @return Gets the instance of this subsystem.
+	 */
+	UFUNCTION(BlueprintPure, Category = "Media Stream|Object")
+	static MEDIASTREAM_API UMediaStreamObjectHandlerSubsystem* Get();
+
+	/**
+	 * Checks whether the given class has a class handler. Checks each super class too.
+	 * @param InClass The class to check.
+	 * @return True if any handler can handle the class.
+	 */
+	UFUNCTION(BlueprintPure, Category = "Media Stream|Object")
+	MEDIASTREAM_API bool CanHandleObject(const UClass* InClass) const;
+
+	/**
+	 * Checks whether a class has a handler registered.
+	 * @param InClass The class to check. Must be an exact match.
+	 * @return True if the handler is registered.
+	 */
+	UFUNCTION(BlueprintPure, Category = "Media Stream|Object")
+	MEDIASTREAM_API bool HasObjectHandler(const UClass* InClass) const;
+
+	/**
+	 * Creates a media player for a given Media Stream and source object.
+	 * @param InMediaStream The Media Stream container.
+	 * @param InPlaybackSource The playback source.
+	 * @param InCurrentPlayer The current player to update or null.
+	 * @return A media player or nullptr on failure.
+	 */
+	UFUNCTION(BlueprintCallable, Category = "Media Stream|Object")
+	MEDIASTREAM_API UMediaPlayer* CreateMediaPlayer(UMediaStream* InMediaStream, UObject* InPlaybackSource, UMediaPlayer* InCurrentPlayer) const;
+};
diff --git a/Engine/Plugins/Experimental/MediaStream/Source/MediaStream/Public/MediaStreamPlayerConfig.h b/Engine/Plugins/Experimental/MediaStream/Source/MediaStream/Public/MediaStreamPlayerConfig.h
new file mode 100644
--- /dev/null
+++ b/Engine/Plugins/Experimental/MediaStream/Source/MediaStream/Public/MediaStreamPlayerConfig.h
@@ -0,0 +1,60 @@
+// Copyright Epic Games, Inc. All Rights Reserved.
+
+#pragma once
+
+#include "MediaPlayer.h"
+
+#include "MediaStreamPlayerConfig.generated.h"
+
+class UMediaPlayer;
+
+USTRUCT(BlueprintType)
+struct FMediaStreamPlayerConfig
+{
+	GENERATED_BODY()
+
+	UPROPERTY(EditInstanceOnly, BlueprintReadWrite, Category = "Media Stream|Player")
+	bool bPlayOnOpen = true;
+
+	UPROPERTY(EditInstanceOnly, BlueprintReadWrite, Category = "Media Stream|Player", EditFixedSize, meta = (EditFixedOrder))
+	FMediaPlayerTrackOptions TrackOptions;
+
+	UPROPERTY(EditInstanceOnly, BlueprintReadWrite, Category = "Media Stream|Player")
+	bool bLooping = false;
+
+	UPROPERTY(EditInstanceOnly, BlueprintReadWrite, Category = "Media Stream|Player", meta = (UIMin = 0, ClampMin = 0, UIMax = 1, ClampMax = 1))
+	float Volume = 1.f;
+
+	UPROPERTY(EditInstanceOnly, BlueprintReadWrite, Category = "Media Stream|Player")
+	TOptional<FFloatInterval> PlaybackTimeRange;
+
+	UPROPERTY(EditInstanceOnly, BlueprintReadWrite, Category = "Media Stream|Player")
+	float Rate = 1.f;
+
+	UPROPERTY(EditInstanceOnly, BlueprintReadWrite, Category = "Media Stream|Player", meta = (Units = "s"))
+	float TimeDelay = 0.f;
+
+	/** Only applicable to play lists. */
+	UPROPERTY(EditInstanceOnly, BlueprintReadWrite, Category = "Media Stream|Player")
+	bool bShuffle = false;
+
+	UPROPERTY(EditInstanceOnly, BlueprintReadWrite, Category = "Media Stream|Player", AdvancedDisplay, 
+		meta = (Units = "s", UIMin = 0.1, ClampMin = 0.01, UIMax = 5, ClampMax = 60))
+	float CacheAhead = 0.1f;
+
+	UPROPERTY(EditInstanceOnly, BlueprintReadWrite, Category = "Media Stream|Player", AdvancedDisplay,
+		meta = (Units = "s", UIMin = 0.1, ClampMin = 0.01, UIMax = 5, ClampMax = 60))
+	float CacheBehind = 3.0f;
+
+	UPROPERTY(EditInstanceOnly, BlueprintReadWrite, Category = "Media Stream|Player", AdvancedDisplay,
+		meta = (Units = "s", UIMin = 0.1, ClampMin = 0.01, UIMax = 5, ClampMax = 60))
+	float CacheBehindGame = 0.1f;
+
+	MEDIASTREAM_API bool operator==(const FMediaStreamPlayerConfig& InOther) const;
+
+	FMediaPlayerOptions CreateOptions(const FTimespan& InStartTime, const TMap<FName, FVariant>& InCustomOptions = {}) const;
+
+	void ApplyConfig(UMediaPlayer& InMediaPlayer) const;
+
+	bool ApplyRate(UMediaPlayer& InMediaPlayer) const;
+};
diff --git a/Engine/Plugins/Experimental/MediaStream/Source/MediaStream/Public/MediaStreamSchemeHandlerManager.h b/Engine/Plugins/Experimental/MediaStream/Source/MediaStream/Public/MediaStreamSchemeHandlerManager.h
new file mode 100644
--- /dev/null
+++ b/Engine/Plugins/Experimental/MediaStream/Source/MediaStream/Public/MediaStreamSchemeHandlerManager.h
@@ -0,0 +1,121 @@
+// Copyright Epic Games, Inc. All Rights Reserved.
+
+#pragma once
+
+#include "Containers/Map.h"
+#include "HAL/Platform.h"
+#include "Templates/SharedPointer.h"
+#include "UObject/NameTypes.h"
+
+class IMediaStreamSchemeHandler;
+class UMediaPlayer;
+class UMediaStream;
+struct FMediaStreamSource;
+
+/**
+ * Handles the registration and operation of Scheme Handlers.
+ */
+class FMediaStreamSchemeHandlerManager
+{
+public:
+	/**
+	 * @return Gets the instance of this manager.
+	 */
+	static MEDIASTREAM_API FMediaStreamSchemeHandlerManager& Get();
+
+	/**
+	 * Checks whether a scheme has a handler registered.
+	 * @param InScheme The scheme to check.
+	 * @return True if the handler is registered.
+	 */
+	MEDIASTREAM_API bool HasSchemeHandler(FName InScheme) const;
+
+	/**
+	 * @return Gets the list of registered scheme handlers.
+	 */
+	MEDIASTREAM_API TArray<FName> GetSchemeHandlerNames() const;
+
+	/**
+	 * Registers a Scheme handler. Holds a strong reference.
+	 * Does not replace an already registered handler.
+	 * @param InScheme The scheme name.
+	 * @param InHandler The handler which converts the path to a stream source.
+	 * @return True if the handler was registered.
+	 */
+	MEDIASTREAM_API bool RegisterSchemeHandler(FName InScheme, const TSharedRef<IMediaStreamSchemeHandler>& InHandler);
+
+	/**
+	 * Templated version of RegisterSchemeHandler.
+	 * @tparam InHandlerClass The handler class to register.
+	 * @tparam InArgsType The type of handle class constructor parameters (automatic).
+	 * @param InScheme The scheme to register.
+	 * @param InArgs The handle class constructor parameters.
+	 * @return True if the handler was registered.
+	 */
+	template<typename InHandlerClass, typename... InArgsType
+		UE_REQUIRES(TIsDerivedFrom<InHandlerClass, IMediaStreamSchemeHandler>::Value)>
+	bool RegisterSchemeHandler(FName InScheme, InArgsType&&... InArgs)
+	{
+		return RegisterSchemeHandler(InScheme, MakeShared<InHandlerClass>(Forward<InArgsType>(InArgs)...));
+	}
+
+	/**
+	 * Templated version of RegisterSchemeHandler with auto scheme value.
+	 * @tparam InHandlerClass The handler class to register.
+	 * @tparam InArgsType The type of handle class constructor parameters (automatic).
+	 * @param InArgs The handle class constructor parameters.
+	 * @return True if the handler was registered.
+	 */
+	template<typename InHandlerClass, typename... InArgsType
+		UE_REQUIRES(TIsDerivedFrom<InHandlerClass, IMediaStreamSchemeHandler>::Value)>
+	bool RegisterSchemeHandler(InArgsType&&... InArgs)
+	{
+		return RegisterSchemeHandler(InHandlerClass::Scheme, MakeShared<InHandlerClass>(Forward<InArgsType>(InArgs)...));
+	}
+
+	/**
+	 * Unregisters a Scheme handler.
+	 * @param InScheme The scheme to unregister.
+	 * @return The removed Scheme handler, if found.
+	 */
+	MEDIASTREAM_API TSharedPtr<IMediaStreamSchemeHandler> UnregisterSchemeHandler(FName InScheme);
+
+	/**
+	 * Templated version of UnregisterSchemeHandler with auto scheme value.
+	 * @tparam InHandlerClass The handler class to register.
+	 * @return The removed Scheme handler, if found.
+	 */
+	template<typename InHandlerClass
+		UE_REQUIRES(TIsDerivedFrom<InHandlerClass, IMediaStreamSchemeHandler>::Value)>
+	TSharedPtr<IMediaStreamSchemeHandler> UnregisterSchemeHandler()
+	{
+		return UnregisterSchemeHandler(InHandlerClass::Scheme);
+	}
+
+	/**
+	 * Finds the handler which is used for the given scheme.
+	 * @param InScheme The scheme to check.
+	 * @return The found handler or nullptr.
+	 */
+	MEDIASTREAM_API TSharedPtr<IMediaStreamSchemeHandler> GetHandlerTypeForScheme(FName InScheme) const;
+
+	/**
+	 * Create a Media Stream Source from a scheme and path. Must have a registered handler.
+	 * @param InOuter The container of the source.
+	 * @param InScheme The scheme used to resolve the path.
+	 * @param InPath The path to the media source. This should not be Url encoded.
+	 * @return A valid stream source object or none if it was not a valid Url.
+	 */
+	MEDIASTREAM_API FMediaStreamSource CreateSource(UObject* InOuter, FName InScheme, const FString& InPath) const;
+
+	/**
+	 * Creates a media player for the given Media Stream based on its source.
+	 * @param InMediaStream The Media Stream to create a player for.
+	 * @param InCurrentPlayer
+	 * @return A media player or nullptr on failure.
+	 */
+	MEDIASTREAM_API UMediaPlayer* CreateOrUpdateMediaPlayer(UMediaStream* InMediaStream, UMediaPlayer* InCurrentPlayer) const;
+
+private:
+	TMap<FName, TSharedRef<IMediaStreamSchemeHandler>> Handlers;
+};
diff --git a/Engine/Plugins/Experimental/MediaStream/Source/MediaStream/Public/MediaStreamSchemeHandlerSubsystem.h b/Engine/Plugins/Experimental/MediaStream/Source/MediaStream/Public/MediaStreamSchemeHandlerSubsystem.h
new file mode 100644
--- /dev/null
+++ b/Engine/Plugins/Experimental/MediaStream/Source/MediaStream/Public/MediaStreamSchemeHandlerSubsystem.h
@@ -0,0 +1,67 @@
+// Copyright Epic Games, Inc. All Rights Reserved.
+
+#pragma once
+
+#include "Subsystems/EngineSubsystem.h"
+
+#include "HAL/Platform.h"
+#include "UObject/NameTypes.h"
+
+#if WITH_EDITOR
+#include "Templates/SharedPointerFwd.h"
+#endif
+
+#include "MediaStreamSchemeHandlerSubsystem.generated.h"
+
+class UMediaPlayer;
+class UMediaStream;
+struct FMediaStreamSource;
+
+/**
+ * Subsystem for blueprint-style interaction with the Scheme Handler Manager.
+ */
+UCLASS(MinimalAPI, BlueprintType, ClassGroup = "Media Stream")
+class UMediaStreamSchemeHandlerSubsystem : public UEngineSubsystem
+{
+	GENERATED_BODY()
+
+public:
+	/**
+	 * @return Gets the instance of this subsystem.
+	 */
+	UFUNCTION(BlueprintPure, Category = "Media Stream|Scheme")
+	static MEDIASTREAM_API UMediaStreamSchemeHandlerSubsystem* Get();
+
+	/**
+	 * Checks whether a scheme has a handler registered.
+	 * @param InScheme The scheme to check.
+	 * @return True if the handler is registered.
+	 */
+	UFUNCTION(BlueprintPure, Category = "Media Stream|Scheme")
+	MEDIASTREAM_API bool HasSchemeHandler(FName InScheme) const;
+
+	/**
+	 * @return Gets the list of registered scheme handlers.
+	 */
+	UFUNCTION(BlueprintPure, Category = "Media Stream|Scheme")
+	MEDIASTREAM_API TArray<FName> GetSchemeHandlerNames() const;
+
+	/**
+	 * Create a Media Stream Source from a scheme and path. Must have a registered handler.
+	 * @param InOuter The container of the source.
+	 * @param InScheme The scheme used to resolve the path.
+	 * @param InPath The path to the media source. This should not be Url encoded.
+	 * @return A valid stream source object or none if it was not a valid Url.
+	 */
+	UFUNCTION(BlueprintCallable, Category = "Media Stream|Scheme")
+	FMediaStreamSource CreateSource(UObject* InOuter, FName InScheme, const FString& InPath) const;
+
+	/**
+	 * Creates a media player for the given Media Stream based on its source.
+	 * @param InMediaStream The Media Stream to create a player for.
+	 * @param InCurrentPlayer
+	 * @return A media player or nullptr on failure.
+	 */
+	UFUNCTION(BlueprintCallable, Category = "Media Stream|Scheme")
+	MEDIASTREAM_API UMediaPlayer* CreateOrUpdateMediaPlayer(UMediaStream* InMediaStream, UMediaPlayer* InCurrentPlayer) const;
+};
diff --git a/Engine/Plugins/Experimental/MediaStream/Source/MediaStream/Public/MediaStreamSource.h b/Engine/Plugins/Experimental/MediaStream/Source/MediaStream/Public/MediaStreamSource.h
new file mode 100644
--- /dev/null
+++ b/Engine/Plugins/Experimental/MediaStream/Source/MediaStream/Public/MediaStreamSource.h
@@ -0,0 +1,61 @@
+// Copyright Epic Games, Inc. All Rights Reserved.
+
+#pragma once
+
+#include "Containers/UnrealString.h"
+#include "Kismet/BlueprintFunctionLibrary.h"
+#include "UObject/ObjectPtr.h"
+#include "UObject/SoftObjectPtr.h"
+
+#include "MediaStreamSource.generated.h"
+
+class UMediaPlayer;
+class UMediaSource;
+class UMediaStream;
+
+/**
+ * Represents the source of a Media Stream.
+ *
+ * This uses the idea of a "scheme" and "path" to reference an asset.
+ * This nomenclature was chosen so that it would be easy to relate url schemes to
+ * this system, such as file://path or http://path.
+ */
+USTRUCT(BlueprintType)
+struct FMediaStreamSource
+{
+	GENERATED_BODY()
+
+	/**
+	 * Used by the handler subsystem to identifer the scheme used.
+	 * A value of NAME_None identifies a null or invalid source.
+	 */
+	UPROPERTY(EditInstanceOnly, BlueprintReadOnly, Category = "Media Stream|Source",
+		meta = (GetOptions = "/Script/MediaStream.MediaStreamSourceBlueprintFunctionLibrary:GetSchemeTypes"))
+	FName Scheme = NAME_None;
+
+	/**
+	 * Path to the referenced source, such as a file/asset path or a managed source name.
+	 * A path must be provided to have a valid scheme.
+	 */
+	UPROPERTY(EditInstanceOnly, BlueprintReadOnly, Category = "Media Stream|Source")
+	FString Path;
+
+	/**
+	 * Source of the Media. Either a UMediaStream, UMediaSource or a UMediaPlaylist
+	 */
+	UPROPERTY(EditInstanceOnly, BlueprintReadOnly, Category = "Media Stream|Source",
+		meta = (AllowedClasses="MediaStream,MediaSource,MediaPlaylist"))
+	TObjectPtr<UObject> Object = nullptr;
+
+	MEDIASTREAM_API bool operator==(const FMediaStreamSource& InOther) const;
+};
+
+UCLASS()
+class UMediaStreamSourceBlueprintFunctionLibrary : public UBlueprintFunctionLibrary
+{
+	GENERATED_BODY()
+
+public:
+	UFUNCTION()
+	static TArray<FName> GetSchemeTypes();
+};
\ No newline at end of file
diff --git a/Engine/Plugins/Experimental/MediaStream/Source/MediaStream/Public/MediaStreamSourceBlueprintLibrary.h b/Engine/Plugins/Experimental/MediaStream/Source/MediaStream/Public/MediaStreamSourceBlueprintLibrary.h
new file mode 100644
--- /dev/null
+++ b/Engine/Plugins/Experimental/MediaStream/Source/MediaStream/Public/MediaStreamSourceBlueprintLibrary.h
@@ -0,0 +1,126 @@
+// Copyright Epic Games, Inc. All Rights Reserved.
+
+#pragma once
+
+#include "Kismet/BlueprintFunctionLibrary.h"
+
+#include "Containers/StringFwd.h"
+#include "MediaStreamSource.h"
+#include "UObject/ObjectPtr.h"
+#include "UObject/SoftObjectPtr.h"
+
+#include "MediaStreamSourceBlueprintLibrary.generated.h"
+
+/**
+ * Deals with creating new Media Stream Sources.
+ */
+UCLASS(MinimalAPI, BlueprintType, ClassGroup = "Media Stream")
+class UMediaStreamSourceBlueprintLibrary : public UBlueprintFunctionLibrary
+{
+	GENERATED_BODY()
+
+public:
+	/**
+	 * Checks the media stream to see if its source type is none or set to something valid.
+	 * And that the path is not empty.
+	 * @param InSource The source to check.
+	 * @return True if the source is valid. Does not guarantee it is correct, only valid.
+	 */
+	UFUNCTION(BlueprintPure, Category = "Media Stream|Source")
+	static MEDIASTREAM_API bool IsValidMediaSource(const FMediaStreamSource& InSource);
+
+	/**
+	 * Checks the Asset pointer to see if it's a valid asset for a Media Stream Source.
+	 * @param InAsset The asset to check.
+	 * @return True if the asset is valid.
+	 */
+	UFUNCTION(BlueprintPure, Category = "Media Stream|Source")
+	static MEDIASTREAM_API bool IsAssetValid(const TSoftObjectPtr<UObject>& InAsset);
+
+	/**
+	 * Checks the Asset path to see if it's a valid asset for a Media Stream Source.
+	 * @param InPath The asset path to check.
+	 * @return True if the asset is valid.
+	 */
+	UFUNCTION(BlueprintPure, Category = "Media Stream|Source")
+	static MEDIASTREAM_API bool IsAssetPathValid(const FString& InPath);
+
+	/**
+	 * Checks the Asset soft path to see if it's a valid asset for a Media Stream Source.
+	 * @param InPath The asset soft path to check.
+	 * @return True if the asset is valid.
+	 */
+	UFUNCTION(BlueprintPure, Category = "Media Stream|Source")
+	static MEDIASTREAM_API bool IsAssetSoftPathValid(const FSoftObjectPath& InPath);
+
+	/**
+	 * Create a Media Stream Source from an Asset. Must have a registered handler.
+	 * @param InMediaStream The owning Media Stream
+	 * @param InScheme The scheme of the source.
+	 * @param InPath The path to the source.
+	 * @return A valid stream source object or none if it was not a valid object.
+	 */
+	UFUNCTION(BlueprintCallable, Category = "Media Stream|Source")
+	static MEDIASTREAM_API FMediaStreamSource MakeMediaSourceFromSchemePath(UMediaStream* InMediaStream, FName InScheme, const FString& InPath);
+
+	/**
+	 * Create a Media Stream Source from an Asset. Must have a registered handler.
+	 * @param InMediaStream The owning Media Stream
+	 * @param InObject The playable object.
+	 * @return A valid stream source object or none if it was not a valid object.
+	 */
+	UFUNCTION(BlueprintCallable, Category = "Media Stream|Source")
+	static MEDIASTREAM_API FMediaStreamSource MakeMediaSourceFromAsset(UMediaStream* InMediaStream, const TSoftObjectPtr<UObject>& InAsset);
+
+	/**
+	 * Create a Media Stream Source from a stream name. Must be set up via the Media Source Manager.
+	 * @param InMediaStream The owning Media Stream
+	 * @param InStreamName The name of the registered stream.
+	 * @return A valid stream source object or none if it was not a valid stream or the stream was not found.
+	 */
+	UFUNCTION(BlueprintCallable, Category = "Media Stream|Source")
+	static MEDIASTREAM_API FMediaStreamSource MakeMediaSourceFromStreamName(UMediaStream* InMediaStream, FName InStreamName);
+
+	/**
+	 * Create a Media Stream Source from a file name.
+	 * @param InMediaStream The owning Media Stream
+	 * @param InFileName The path to the file name. Should be a full path or relative to the project or engine root.
+	 * @return A valid stream source object or none if the file was not found.
+	 */
+	UFUNCTION(BlueprintCallable, Category = "Media Stream|Source")
+	static MEDIASTREAM_API FMediaStreamSource MakeMediaSourceFromFile(UMediaStream* InMediaStream, const FString& InFileName);
+
+	/**
+	 * Create a Media Stream Source from an object.
+	 * @param InMediaStream The owning Media Stream
+	 * @param InRootObject Root of the relative path.
+	 * @param InObject The playable object. Should share a common ancestor with the Media Stream, but this is not enforced.
+	 * @return A valid stream source object or none if the file was not found.
+	 */
+	UFUNCTION(BlueprintCallable, Category = "Media Stream|Source")
+	static MEDIASTREAM_API FMediaStreamSource MakeMediaSourceFromSubobject(UMediaStream* InMediaStream, UObject* InObject);
+
+	/**
+	 * Create a Media Stream Source by instantiating the given class.
+	 * @param InMediaStream The owning Media Stream
+	 * @param InClass The class of a playable object.
+	 * @param InOuter The outer of the new object. Defaults to the media stream object.
+	 * @return A valid stream source object or none if the file was not found.
+	 */
+	UFUNCTION(BlueprintCallable, Category = "Media Stream|Source")
+	static MEDIASTREAM_API FMediaStreamSource MakeMediaSourceFromSubobjectClass(UMediaStream* InMediaStream, const UClass* InClass);
+
+	/**
+	 * Templated version of MakeMediaSourceFromClass.
+	 * @tparam InClassName The name of the UClass to create an playable from.
+	 * @param InMediaStream The owning Media Stream
+	 * @param InOuter The outer of the newly created object
+	 * @return A valid stream source object or none if the file was not found.
+	 */
+	template<typename InClassName
+		UE_REQUIRES(TModels_V<CStaticClassProvider, InClassName>)>
+	static FMediaStreamSource MakeMediaSourceFromSubobject(UMediaStream* InMediaStream)
+	{
+		return MakeMediaSourceFromSubobjectClass(InMediaStream, InClassName::StaticClass());
+	}
+};
diff --git a/Engine/Plugins/Experimental/MediaStream/Source/MediaStream/Public/MediaStreamTextureConfig.h b/Engine/Plugins/Experimental/MediaStream/Source/MediaStream/Public/MediaStreamTextureConfig.h
new file mode 100644
--- /dev/null
+++ b/Engine/Plugins/Experimental/MediaStream/Source/MediaStream/Public/MediaStreamTextureConfig.h
@@ -0,0 +1,25 @@
+// Copyright Epic Games, Inc. All Rights Reserved.
+
+#pragma once
+
+#include "MediaTexture.h"
+
+#include "MediaStreamTextureConfig.generated.h"
+
+class UMediaTexture;
+
+USTRUCT(BlueprintType)
+struct FMediaStreamTextureConfig
+{
+	GENERATED_BODY()
+
+	UPROPERTY(EditInstanceOnly, BlueprintReadWrite, Category = "Media Stream|Texture")
+	bool bEnableMipGen = false;
+
+	UPROPERTY(EditInstanceOnly, BlueprintReadWrite, Category = "Media Stream|Texture")
+	uint8 MipGenCount = 1;
+
+	MEDIASTREAM_API bool operator==(const FMediaStreamTextureConfig& InOther) const;
+
+	void ApplyConfig(UMediaTexture& InMediaTexture) const;
+};
diff --git a/Engine/Plugins/Experimental/MediaStream/Source/MediaStreamEditor/MediaStreamEditor.Build.cs b/Engine/Plugins/Experimental/MediaStream/Source/MediaStreamEditor/MediaStreamEditor.Build.cs
new file mode 100644
--- /dev/null
+++ b/Engine/Plugins/Experimental/MediaStream/Source/MediaStreamEditor/MediaStreamEditor.Build.cs
@@ -0,0 +1,40 @@
+// Copyright Epic Games, Inc. All Rights Reserved.
+
+using UnrealBuildTool;
+
+public class MediaStreamEditor : ModuleRules
+{
+	public MediaStreamEditor(ReadOnlyTargetRules Target) : base(Target)
+	{
+		PCHUsage = ModuleRules.PCHUsageMode.UseExplicitOrSharedPCHs;
+
+		PublicDependencyModuleNames.AddRange(
+			new string[]
+			{
+				"Core",
+				"Engine"
+			}
+		);
+
+		PrivateDependencyModuleNames.AddRange(
+			new string[]
+			{
+				"CoreUObject",
+				"DesktopWidgets",
+				"LevelSequence",
+				"LevelSequenceEditor",
+				"MediaAssets",
+				"MediaCompositing",
+				"MediaPlayerEditor",
+				"MediaStream",
+				"MediaUtils",
+				"MovieScene",
+				"Projects",
+				"PropertyEditor",
+				"Slate",
+				"SlateCore",
+				"ToolMenus",
+			}
+		);
+	}
+}
diff --git a/Engine/Plugins/Experimental/MediaStream/Source/MediaStreamEditor/Private/DetailsPanel/MediaStreamCustomization.cpp b/Engine/Plugins/Experimental/MediaStream/Source/MediaStreamEditor/Private/DetailsPanel/MediaStreamCustomization.cpp
new file mode 100644
--- /dev/null
+++ b/Engine/Plugins/Experimental/MediaStream/Source/MediaStreamEditor/Private/DetailsPanel/MediaStreamCustomization.cpp
@@ -0,0 +1,550 @@
+// Copyright Epic Games, Inc. All Rights Reserved.
+
+#include "MediaStreamCustomization.h"
+
+#include "DetailCategoryBuilder.h"
+#include "DetailLayoutBuilder.h"
+#include "DetailWidgetRow.h"
+#include "IDetailGroup.h"
+#include "IMediaStreamPlayer.h"
+#include "MediaPlayerEditorModule.h"
+#include "MediaStream.h"
+#include "MediaStreamEditorSequencerLibrary.h"
+#include "MediaStreamEditorStyle.h"
+#include "MediaStreamEnums.h"
+#include "MediaStreamPlayerConfig.h"
+#include "Modules/ModuleManager.h"
+#include "Styling/StyleColors.h"
+#include "Widgets/Input/SButton.h"
+#include "Widgets/Images/SImage.h"
+#include "Widgets/Text/STextBlock.h"
+
+#define LOCTEXT_NAMESPACE "MediaStreamCustomization"
+
+FMediaStreamCustomization::FMediaStreamCustomization()
+{
+}
+
+void FMediaStreamCustomization::CustomizeDetails(IDetailLayoutBuilder& InDetailBuilder)
+{
+	// Get objects we are editing.
+	TArray<TWeakObjectPtr<UObject>> Objects;
+	InDetailBuilder.GetObjectsBeingCustomized(Objects);
+	MediaStreamsList.Reserve(Objects.Num());
+
+	for (TWeakObjectPtr<UObject>& Obj : Objects)
+	{
+		TWeakObjectPtr<UMediaStream> MediaStream = Cast<UMediaStream>(Obj.Get());
+		if (MediaStream.IsValid())
+		{
+			MediaStreamsList.Add(MediaStream);
+		}
+	}
+
+	if (MediaStreamsList.IsEmpty())
+	{
+		return;
+	}
+
+	IDetailCategoryBuilder& MediaStreamCategory = InDetailBuilder.EditCategory("Media Stream");
+
+	AddControlControls(MediaStreamCategory);
+}
+
+void FMediaStreamCustomization::AddControlControls(IDetailCategoryBuilder& InMediaStreamCategory)
+{
+	IDetailGroup& ControlGroup = InMediaStreamCategory.AddGroup(TEXT("Control"), LOCTEXT("Control", "Control"));
+
+	TWeakPtr<FMediaStreamCustomization> WeakSelf = StaticCastWeakPtr<FMediaStreamCustomization>(AsWeak());
+	const ISlateStyle* Stream = &FMediaStreamEditorStyle::Get().Get();
+
+	// Add media player playback slider
+	if (IMediaPlayerEditorModule* MediaPlayerEditorModule = FModuleManager::LoadModulePtr<IMediaPlayerEditorModule>("MediaPlayerEditor"))
+	{
+		const TSharedRef<IMediaPlayerSlider> MediaPlayerSlider =
+			MediaPlayerEditorModule->CreateMediaPlayerSliderWidget(GetMediaPlayers());
+
+		MediaPlayerSlider->SetSliderHandleColor(FSlateColor(EStyleColor::AccentBlue));
+		MediaPlayerSlider->SetVisibleWhenInactive(EVisibility::Visible);
+
+		ControlGroup.AddWidgetRow()
+			.IsEnabled(TAttribute<bool>::CreateSPLambda(this, 
+				[this]()
+				{
+					if (UMediaStream* MediaStream = GetMediaStream())
+					{
+						return !FMediaStreamEditorSequencerLibrary::HasTrack(MediaStream);
+					}
+
+					return false;
+				}))
+			[
+				MediaPlayerSlider
+			];
+	}
+
+	// Add media control buttons.
+	ControlGroup.AddWidgetRow()
+	[
+		SNew(SHorizontalBox)
+		// Rewind button.
+		+ SHorizontalBox::Slot()
+		.VAlign(VAlign_Center)
+		.AutoWidth()
+		[
+			SNew(SButton)
+			.VAlign(VAlign_Center)
+			.IsEnabled_Lambda([WeakSelf]
+				{
+					if (const TSharedPtr<FMediaStreamCustomization> Self = WeakSelf.Pin())
+					{
+						if (UMediaStream* MediaStream = Self->GetMediaStream())
+						{
+							if (!FMediaStreamEditorSequencerLibrary::HasTrack(MediaStream))
+							{
+								if (const UMediaPlayer* MediaPlayer = Self->GetMediaPlayer())
+								{
+									return MediaPlayer->IsReady() &&
+										MediaPlayer->SupportsSeeking() &&
+										MediaPlayer->GetTime() > FTimespan::Zero();
+								}
+							}
+						}
+					}
+
+					return false;
+				})
+			.OnClicked_Lambda([WeakSelf]() -> FReply
+				{
+					if (const TSharedPtr<FMediaStreamCustomization> Self = WeakSelf.Pin())
+					{
+						Self->OnPlaybackSeek(EMediaStreamPlaybackSeek::Previous);
+						return FReply::Handled();
+					}
+
+					return FReply::Unhandled();
+				})
+			[
+				SNew(SImage)
+				.ColorAndOpacity(FSlateColor::UseForeground())
+				.Image(Stream->GetBrush("MediaStreamEditor.RewindMedia.Small"))
+				.ToolTipText(LOCTEXT("Rewind", "Rewind the media to the beginning"))
+			]
+		]
+
+		// Reverse button.
+		+ SHorizontalBox::Slot()
+		.VAlign(VAlign_Center)
+		.AutoWidth()
+		[
+			SNew(SButton)
+			.VAlign(VAlign_Center)
+			.IsEnabled_Lambda([WeakSelf]
+				{
+					if (const TSharedPtr<FMediaStreamCustomization> Self = WeakSelf.Pin())
+					{
+						if (UMediaStream* MediaStream = Self->GetMediaStream())
+						{
+							if (!FMediaStreamEditorSequencerLibrary::HasTrack(MediaStream))
+							{
+								if (UMediaPlayer* MediaPlayer = Self->GetMediaPlayer())
+								{
+									const float Rate = FMath::Abs(MediaPlayer->GetRate());
+
+									return MediaPlayer->IsReady() && Self->IsRateSupported(Rate > 0 ? -Rate : -1.f);
+								}
+							}
+						}
+					}
+
+					return false;
+				})
+			.OnClicked_Lambda([WeakSelf]() -> FReply
+				{
+					if (const TSharedPtr<FMediaStreamCustomization> Self = WeakSelf.Pin())
+					{
+						Self->OnChangePlaybackDirection(EMediaStreamPlaybackDirection::Backward);
+						return FReply::Handled();
+					}
+
+					return FReply::Unhandled();
+				})
+			[
+				SNew(SImage)
+				.ColorAndOpacity(FSlateColor::UseForeground())
+				.Image(Stream->GetBrush("MediaStreamEditor.ReverseMedia.Small"))
+				.ToolTipText(LOCTEXT("Reverse", "Reverse media playback"))
+			]
+		]
+
+		// Play button.
+		+ SHorizontalBox::Slot()
+		.VAlign(VAlign_Center)
+		.AutoWidth()
+		[
+			SNew(SButton)
+			.VAlign(VAlign_Center)
+			.IsEnabled_Lambda([WeakSelf]
+				{
+					if (const TSharedPtr<FMediaStreamCustomization> Self = WeakSelf.Pin())
+					{
+						if (UMediaStream* MediaStream = Self->GetMediaStream())
+						{
+							if (!FMediaStreamEditorSequencerLibrary::HasTrack(MediaStream))
+							{
+								// Is the player paused or fast forwarding/rewinding?
+								if (const UMediaPlayer* MediaPlayer = Self->GetMediaPlayer())
+								{
+									return MediaPlayer->IsReady()
+										&& (!MediaPlayer->IsPlaying() || (MediaPlayer->GetRate() != 1.0f));
+								}
+							}
+						}
+					}
+
+					return false;
+				})
+			.Visibility_Lambda([WeakSelf]
+				{
+					if (const TSharedPtr<FMediaStreamCustomization> Self = WeakSelf.Pin())
+					{
+						if (const UMediaPlayer* MediaPlayer = Self->GetMediaPlayer())
+						{
+							return MediaPlayer->IsPlaying() ? EVisibility::Collapsed : EVisibility::Visible;
+						}
+					}
+
+					return EVisibility::Visible;
+				})
+			.OnClicked_Lambda([WeakSelf]() -> FReply
+				{
+					if (const TSharedPtr<FMediaStreamCustomization> Self = WeakSelf.Pin())
+					{
+						Self->OnChangePlaybackState(EMediaStreamPlaybackState::Play);
+						return FReply::Handled();
+					}
+
+					return FReply::Unhandled();
+				})
+			[
+				SNew(SImage)
+				.ColorAndOpacity(FSlateColor::UseForeground())
+				.Image(Stream->GetBrush("MediaStreamEditor.PlayMedia.Small"))
+				.ToolTipText(LOCTEXT("Play", "Start media playback"))
+			]
+		]
+
+		// Pause button.
+		+ SHorizontalBox::Slot()
+		.VAlign(VAlign_Center)
+		.AutoWidth()
+		[
+			SNew(SButton)
+			.VAlign(VAlign_Center)
+			.IsEnabled_Lambda([WeakSelf]
+				{
+					if (const TSharedPtr<FMediaStreamCustomization> Self = WeakSelf.Pin())
+					{
+						if (UMediaStream* MediaStream = Self->GetMediaStream())
+						{
+							if (!FMediaStreamEditorSequencerLibrary::HasTrack(MediaStream))
+							{
+								if (const UMediaPlayer* MediaPlayer = Self->GetMediaPlayer())
+								{
+									return MediaPlayer->CanPause() && !MediaPlayer->IsPaused();
+								}
+							}
+						}
+					}
+
+					return false;
+				})
+			.Visibility_Lambda([WeakSelf]
+				{
+					if (const TSharedPtr<FMediaStreamCustomization> Self = WeakSelf.Pin())
+					{
+						if (const UMediaPlayer* MediaPlayer = Self->GetMediaPlayer())
+						{
+							const bool bIsVisible = MediaPlayer->CanPause() && !MediaPlayer->IsPaused();
+							return bIsVisible ? EVisibility::Visible : EVisibility::Collapsed;
+						}
+					}
+
+					return EVisibility::Collapsed;
+				})
+			.OnClicked_Lambda([WeakSelf]() -> FReply
+				{
+					if (const TSharedPtr<FMediaStreamCustomization> Self = WeakSelf.Pin())
+					{
+						Self->OnChangePlaybackState(EMediaStreamPlaybackState::Pause);
+						return FReply::Handled();
+					}
+
+					return FReply::Unhandled();
+				})
+			[
+				SNew(SImage)
+				.ColorAndOpacity(FSlateColor::UseForeground())
+				.Image(Stream->GetBrush("MediaStreamEditor.PauseMedia.Small"))
+				.ToolTipText(LOCTEXT("Pause", "Pause media playback"))
+			]
+		]
+
+		// Forward button.
+		+ SHorizontalBox::Slot()
+		.VAlign(VAlign_Center)
+		.AutoWidth()
+		[
+			SNew(SButton)
+			.VAlign(VAlign_Center)
+			.IsEnabled_Lambda([WeakSelf]
+				{
+					if (const TSharedPtr<FMediaStreamCustomization> Self = WeakSelf.Pin())
+					{
+						if (UMediaStream* MediaStream = Self->GetMediaStream())
+						{
+							if (!FMediaStreamEditorSequencerLibrary::HasTrack(MediaStream))
+							{
+								if (UMediaPlayer* MediaPlayer = Self->GetMediaPlayer())
+								{
+									return MediaPlayer->IsReady();
+								}
+							}
+						}
+					}
+
+					return false;
+				})
+			.OnClicked_Lambda([WeakSelf]() -> FReply
+				{
+					if (const TSharedPtr<FMediaStreamCustomization> Self = WeakSelf.Pin())
+					{
+						Self->OnPlaybackSeek(EMediaStreamPlaybackSeek::End);
+						return FReply::Handled();
+					}
+
+					return FReply::Unhandled();
+				})
+			[
+				SNew(SImage)
+				.ColorAndOpacity(FSlateColor::UseForeground())
+				.Image(Stream->GetBrush("MediaStreamEditor.ForwardMedia.Small"))
+				.RenderTransform(FSlateRenderTransform(TTransform2<float>(TQuat2<float>(PI), FVector2f(20.f, 20.f))))
+				.ToolTipText(LOCTEXT("Forward", "Fast forward media playback"))
+			]
+		]
+
+		// Add Media Track
+		+ SHorizontalBox::Slot()
+		.VAlign(VAlign_Center)
+		.AutoWidth()
+		[
+			SNew(SButton)
+			.VAlign(VAlign_Center)
+			.IsEnabled_Lambda([WeakSelf]
+				{
+					if (const TSharedPtr<FMediaStreamCustomization> Self = WeakSelf.Pin())
+					{
+						if (UMediaStream* MediaStream = Self->GetMediaStream())
+						{
+							return FMediaStreamEditorSequencerLibrary::CanAddTrack(MediaStream);
+						}
+					}
+
+					return false;
+				})
+			.OnClicked_Lambda([WeakSelf]() -> FReply
+				{
+					if (const TSharedPtr<FMediaStreamCustomization> Self = WeakSelf.Pin())
+					{
+						if (Self->MediaStreamsList.Num() == 1)
+						{
+							if (UMediaStream* MediaStream = Self->MediaStreamsList[0].Get())
+							{
+								FMediaStreamEditorSequencerLibrary::AddTrack(MediaStream);
+							}
+						}
+					}
+
+					return FReply::Unhandled();
+				})
+			[
+				SNew(STextBlock)
+				.Font(IDetailLayoutBuilder::GetDetailFont())
+				.Text(LOCTEXT("AddTrack", "Add Track"))
+			]
+		]
+	];
+}
+
+void FMediaStreamCustomization::OnChangePlaybackState(EMediaStreamPlaybackState InState)
+{
+	for (TWeakObjectPtr<UMediaStream>& MediaStreamPtr : MediaStreamsList)
+	{
+		if (UMediaStream* MediaStream = MediaStreamPtr.Get())
+		{
+			if (IMediaStreamPlayer* MediaStreamPlayer = MediaStream->GetPlayer().GetInterface())
+			{
+				if (MediaStreamPlayer->IsReadOnly())
+				{
+					continue;
+				}
+
+				switch (InState)
+				{
+					case EMediaStreamPlaybackState::Play:
+					{
+						const float RequestedRate = FMath::Abs(MediaStreamPlayer->GetPlayerConfig().Rate);
+
+						if (RequestedRate != MediaStreamPlayer->GetPlayerConfig().Rate)
+						{
+							FMediaStreamPlayerConfig PlayerConfig = MediaStreamPlayer->GetPlayerConfig();
+							PlayerConfig.Rate = RequestedRate;
+
+							MediaStreamPlayer->SetPlayerConfig(PlayerConfig);
+						}
+
+						MediaStreamPlayer->Play();
+						break;
+					}
+
+					case EMediaStreamPlaybackState::Pause:
+						MediaStreamPlayer->Pause();
+						break;
+				}
+			}
+		}
+	}
+}
+
+void FMediaStreamCustomization::OnChangePlaybackDirection(EMediaStreamPlaybackDirection InDirection)
+{
+	for (TWeakObjectPtr<UMediaStream>& MediaStreamPtr : MediaStreamsList)
+	{
+		if (UMediaStream* MediaStream = MediaStreamPtr.Get())
+		{
+			if (IMediaStreamPlayer* MediaStreamPlayer = MediaStream->GetPlayer().GetInterface())
+			{
+				if (MediaStreamPlayer->IsReadOnly())
+				{
+					continue;
+				}
+
+				const float RequestedRate = FMath::Abs(MediaStreamPlayer->GetPlayerConfig().Rate)
+					* (InDirection == EMediaStreamPlaybackDirection::Forward ? 1.f : -1.f);
+
+				if (RequestedRate != MediaStreamPlayer->GetPlayerConfig().Rate)
+				{
+					FMediaStreamPlayerConfig PlayerConfig = MediaStreamPlayer->GetPlayerConfig();
+					PlayerConfig.Rate = RequestedRate;
+
+					MediaStreamPlayer->SetPlayerConfig(PlayerConfig);
+				}
+			}
+		}
+	}
+}
+
+void FMediaStreamCustomization::OnPlaybackSeek(EMediaStreamPlaybackSeek InPosition)
+{
+	for (TWeakObjectPtr<UMediaStream>& MediaStreamPtr : MediaStreamsList)
+	{
+		if (UMediaStream* MediaStream = MediaStreamPtr.Get())
+		{
+			if (IMediaStreamPlayer* MediaStreamPlayer = MediaStream->GetPlayer().GetInterface())
+			{
+				if (MediaStreamPlayer->IsReadOnly())
+				{
+					continue;
+				}
+
+				switch (InPosition)
+				{
+					case EMediaStreamPlaybackSeek::Previous:
+						MediaStreamPlayer->Previous();
+						break;
+
+					case EMediaStreamPlaybackSeek::Start:
+						MediaStreamPlayer->Rewind();
+						break;
+
+					case EMediaStreamPlaybackSeek::End:
+						MediaStreamPlayer->FastForward();
+						break;
+
+					case EMediaStreamPlaybackSeek::Next:
+						MediaStreamPlayer->Next();
+						break;
+				}
+			}
+		}
+	}
+}
+
+UMediaStream* FMediaStreamCustomization::GetMediaStream() const
+{
+	for (const TWeakObjectPtr<UMediaStream>& MediaStreamPtr : MediaStreamsList)
+	{
+		if (UMediaStream* MediaStream = MediaStreamPtr.Get())
+		{
+			return MediaStream;
+		}
+	}
+
+	return nullptr;
+}
+
+UMediaPlayer* FMediaStreamCustomization::GetMediaPlayer() const
+{
+	for (const TWeakObjectPtr<UMediaStream>& MediaStreamPtr : MediaStreamsList)
+	{
+		if (UMediaStream* MediaStream = MediaStreamPtr.Get())
+		{
+			if (IMediaStreamPlayer* MediaStreamPlayer = MediaStream->GetPlayer().GetInterface())
+			{				
+				if (UMediaPlayer* MediaPlayer = MediaStreamPlayer->GetPlayer())
+				{
+					return MediaPlayer;
+				}
+			}
+		}
+	}
+
+	return nullptr;
+}
+
+TArray<TWeakObjectPtr<UMediaPlayer>> FMediaStreamCustomization::GetMediaPlayers() const
+{
+	TArray<TWeakObjectPtr<UMediaPlayer>> MediaPlayers;
+	MediaPlayers.Reserve(MediaStreamsList.Num());
+
+	for (const TWeakObjectPtr<UMediaStream>& MediaStreamPtr : MediaStreamsList)
+	{
+		if (UMediaStream* MediaStream = MediaStreamPtr.Get())
+		{
+			if (IMediaStreamPlayer* MediaStreamPlayer = MediaStream->GetPlayer().GetInterface())
+			{
+				if (UMediaPlayer* MediaPlayer = MediaStreamPlayer->GetPlayer())
+				{
+					MediaPlayers.Add(MediaPlayer);
+				}
+			}
+		}
+	}
+
+	return MediaPlayers;
+}
+
+bool FMediaStreamCustomization::IsRateSupported(float InRate) const
+{
+	UMediaPlayer* MediaPlayer = GetMediaPlayer();
+
+	if (!MediaPlayer)
+	{
+		return false;
+	}
+
+	return MediaPlayer->SupportsRate(InRate, /* Unthinned */ true);
+}
+
+#undef LOCTEXT_NAMESPACE
diff --git a/Engine/Plugins/Experimental/MediaStream/Source/MediaStreamEditor/Private/DetailsPanel/MediaStreamCustomization.h b/Engine/Plugins/Experimental/MediaStream/Source/MediaStreamEditor/Private/DetailsPanel/MediaStreamCustomization.h
new file mode 100644
--- /dev/null
+++ b/Engine/Plugins/Experimental/MediaStream/Source/MediaStreamEditor/Private/DetailsPanel/MediaStreamCustomization.h
@@ -0,0 +1,75 @@
+// Copyright Epic Games, Inc. All Rights Reserved.
+
+#pragma once
+
+#include "IDetailCustomization.h"
+
+#include "Containers/Array.h"
+#include "UObject/WeakObjectPtr.h"
+
+class IDetailCategoryBuilder;
+class IDetailLayoutBuilder;
+class IPropertyHandle;
+class SWidget;
+class ULevelSequence;
+class UMediaPlayer;
+class UMediaStream;
+enum class EMediaStreamPlaybackDirection : uint8;
+enum class EMediaStreamPlaybackSeek : uint8;
+enum class EMediaStreamPlaybackState : uint8;
+
+/**
+ * Implements a details view customization for the UMediaStreamComponent class.
+ */
+class FMediaStreamCustomization : public IDetailCustomization
+{
+public:
+	/**
+	 * Creates an instance of this class.
+	 *
+	 * @return The new instance.
+	 */
+	static TSharedRef<IDetailCustomization> MakeInstance()
+	{
+		return MakeShared<FMediaStreamCustomization>();
+	}
+
+	FMediaStreamCustomization();
+	virtual ~FMediaStreamCustomization() override = default;
+
+	//~ Begin IDetailCustomization
+	virtual void CustomizeDetails(IDetailLayoutBuilder& InDetailBuilder) override;
+	//~ End IDetailCustomization
+
+private:
+	/** List of the media streams we are editing. */
+	TArray<TWeakObjectPtr<UMediaStream>> MediaStreamsList;
+
+	/**
+	 * Adds the control... controls.
+	 */
+	void AddControlControls(IDetailCategoryBuilder& InMediaStreamCategory);
+
+	/**
+	* Changes the state of selected media streams.
+	*/
+	void OnChangePlaybackState(EMediaStreamPlaybackState InState);
+
+	/**
+	* Changes the state of selected media streams.
+	*/
+	void OnChangePlaybackDirection(EMediaStreamPlaybackDirection InDirection);
+
+	/**
+	* Changes the state of selected media streams.
+	*/
+	void OnPlaybackSeek(EMediaStreamPlaybackSeek InPosition);
+
+	UMediaStream* GetMediaStream() const;
+
+	UMediaPlayer* GetMediaPlayer() const;
+
+	TArray<TWeakObjectPtr<UMediaPlayer>> GetMediaPlayers() const;
+
+	bool IsRateSupported(float InRate) const;
+};
diff --git a/Engine/Plugins/Experimental/MediaStream/Source/MediaStreamEditor/Private/DetailsPanel/MediaStreamSourceCustomization.cpp b/Engine/Plugins/Experimental/MediaStream/Source/MediaStreamEditor/Private/DetailsPanel/MediaStreamSourceCustomization.cpp
new file mode 100644
--- /dev/null
+++ b/Engine/Plugins/Experimental/MediaStream/Source/MediaStreamEditor/Private/DetailsPanel/MediaStreamSourceCustomization.cpp
@@ -0,0 +1,148 @@
+// Copyright Epic Games, Inc. All Rights Reserved.
+
+#include "MediaStreamSourceCustomization.h"
+
+#include "DetailLayoutBuilder.h"
+#include "DetailWidgetRow.h"
+#include "IDetailChildrenBuilder.h"
+#include "IDetailGroup.h"
+#include "IMediaStreamSchemeHandler.h"
+#include "MediaStream.h"
+#include "MediaStreamSchemeHandlerManager.h"
+#include "Widgets/Input/SButton.h"
+#include "Widgets/Input/SCheckBox.h"
+#include "Widgets/Layout/SGridPanel.h"
+#include "Widgets/Text/STextBlock.h"
+
+#define LOCTEXT_NAMESPACE "MediaStreamSourceCustomization"
+
+FMediaStreamSourceCustomization::FMediaStreamSourceCustomization()
+{
+}
+
+void FMediaStreamSourceCustomization::CustomizeHeader(TSharedRef<IPropertyHandle> InPropertyHandle, FDetailWidgetRow& InHeaderRow,
+	IPropertyTypeCustomizationUtils& InCustomizationUtils)
+{
+}
+
+void FMediaStreamSourceCustomization::CustomizeChildren(TSharedRef<IPropertyHandle> InPropertyHandle, IDetailChildrenBuilder& InChildBuilder,
+	IPropertyTypeCustomizationUtils& InCustomizationUtils)
+{
+	PropertyHandle = InPropertyHandle;
+
+	if (!FMediaStreamSchemeHandlerLibrary::GetStreamSourcePtr(InPropertyHandle))
+	{
+		return;
+	}
+
+	IDetailGroup& SourceGroup = InChildBuilder.AddGroup(TEXT("Source"), LOCTEXT("Source", "Source"));
+
+	AddSourceSchemeSelector(SourceGroup, InCustomizationUtils);
+	AddSchemeCustomizations(InPropertyHandle, SourceGroup, InCustomizationUtils);
+}
+
+void FMediaStreamSourceCustomization::AddSourceSchemeSelector(IDetailGroup& InSourceGroup, IPropertyTypeCustomizationUtils& InCustomizationUtils)
+{
+	FDetailWidgetRow& SourceSchemeRow = InSourceGroup.AddWidgetRow();
+
+	SourceSchemeRow.NameContent()
+	[
+		SNew(STextBlock)
+		.Font(IDetailLayoutBuilder::GetDetailFont())
+		.Text(LOCTEXT("SourceSchemeName", "Source Scheme"))
+	];
+
+	TSharedRef<SGridPanel> SchemeSelector = MakeShared<SGridPanel>();
+
+	AddSourceSchemeSelectRow(SchemeSelector, 0, TEXT("None"));
+	AddSourceSchemeSelectRow(SchemeSelector, 1, TEXT("File"));
+	AddSourceSchemeSelectRow(SchemeSelector, 2, TEXT("Asset"));
+
+	/**
+	 * Later implementation
+	 *
+	const TArray<FName> Schemes = FMediaStreamSchemeHandlerManager::Get().GetSchemeHandlerNames();
+	int32 Row = 1;
+
+	for (const FName& Scheme : Schemes)
+	{
+		AddSourceSchemeSelectRow(SchemeSelector, Row, Scheme);
+		++Row;
+	}
+	 */
+
+	SourceSchemeRow.ValueContent()
+	[
+		SchemeSelector
+	];
+}
+
+void FMediaStreamSourceCustomization::AddSourceSchemeSelectRow(TSharedRef<SGridPanel> InContainer, int32 InRow, const FName& InName)
+{
+	InContainer->AddSlot(0, InRow)
+		.HAlign(EHorizontalAlignment::HAlign_Left)
+		.VAlign(EVerticalAlignment::VAlign_Center)
+		.Padding(5.f, 2.f)
+		[
+			SNew(SCheckBox)
+			.Style(FAppStyle::Get(), "Menu.RadioButton")
+			.IsChecked(this, &FMediaStreamSourceCustomization::GetSourceCheckBoxState, InName)
+			.OnCheckStateChanged(this, &FMediaStreamSourceCustomization::OnSourceCheckBoxStateChanged, InName)
+		];
+
+	InContainer->AddSlot(1, InRow)
+		.HAlign(EHorizontalAlignment::HAlign_Left)
+		.VAlign(EVerticalAlignment::VAlign_Center)
+		.Padding(0.f, 2.f, 5.f, 2.f)
+		[
+			SNew(SButton)
+			.HAlign(HAlign_Center)
+			.VAlign(VAlign_Center)
+			.ContentPadding(FMargin(5.f, 1.f))
+			.ButtonStyle(FAppStyle::Get(), TEXT("NoBorder"))
+			[
+				SNew(STextBlock)
+				.Font(IDetailLayoutBuilder::GetDetailFont())
+				.Text(FText::FromName(InName))
+			]
+		];
+}
+
+ECheckBoxState FMediaStreamSourceCustomization::GetSourceCheckBoxState(FName InSourceScheme) const
+{
+	return FMediaStreamSchemeHandlerLibrary::GetScheme(PropertyHandle) == InSourceScheme
+		? ECheckBoxState::Checked
+		: ECheckBoxState::Unchecked;
+}
+
+void FMediaStreamSourceCustomization::OnSourceCheckBoxStateChanged(ECheckBoxState InState, FName InSourceScheme)
+{
+	if (PropertyHandle.IsValid())
+	{
+		FMediaStreamSchemeHandlerLibrary::SetSource(PropertyHandle.ToSharedRef(), InSourceScheme, "");
+	}
+}
+
+FReply FMediaStreamSourceCustomization::OnSourceSchemeButtonClicked(const FName& InSourceScheme)
+{
+	if (PropertyHandle.IsValid())
+	{
+		FMediaStreamSchemeHandlerLibrary::SetSource(PropertyHandle.ToSharedRef(), InSourceScheme, "");
+	}
+
+	return FReply::Handled();
+}
+
+void FMediaStreamSourceCustomization::AddSchemeCustomizations(TSharedRef<IPropertyHandle> InPropertyHandle, IDetailGroup& InSourceGroup,
+	IPropertyTypeCustomizationUtils& InCustomizationUtils)
+{
+	const FMediaStreamSchemeHandlerManager& Manager = FMediaStreamSchemeHandlerManager::Get();
+	const TArray<FName> Schemes = Manager.GetSchemeHandlerNames();
+
+	for (const FName& Scheme : Schemes)
+	{
+		Manager.GetHandlerTypeForScheme(Scheme)->CreatePropertyCustomization(InPropertyHandle, InSourceGroup, InCustomizationUtils);
+	}
+}
+
+#undef LOCTEXT_NAMESPACE
diff --git a/Engine/Plugins/Experimental/MediaStream/Source/MediaStreamEditor/Private/DetailsPanel/MediaStreamSourceCustomization.h b/Engine/Plugins/Experimental/MediaStream/Source/MediaStreamEditor/Private/DetailsPanel/MediaStreamSourceCustomization.h
new file mode 100644
--- /dev/null
+++ b/Engine/Plugins/Experimental/MediaStream/Source/MediaStreamEditor/Private/DetailsPanel/MediaStreamSourceCustomization.h
@@ -0,0 +1,58 @@
+// Copyright Epic Games, Inc. All Rights Reserved.
+
+#pragma once
+
+#include "IPropertyTypeCustomization.h"
+
+#include "Templates/SharedPointer.h"
+
+class FReply;
+class IDetailCategoryBuilder;
+class IDetailGroup;
+class IDetailLayoutBuilder;
+class IPropertyHandle;
+class SGridPanel;
+class SWidget;
+class UMediaPlayer;
+class UMediaStream;
+enum class ECheckBoxState : uint8;
+struct FMediaStreamSource;
+
+/**
+ * Implements a details view customization for the UMediaStreamComponent class.
+ */
+class FMediaStreamSourceCustomization : public IPropertyTypeCustomization
+{
+public:
+	/**
+	 * Creates an instance of this class.
+	 *
+	 * @return The new instance.
+	 */
+	static TSharedRef<IPropertyTypeCustomization> MakeInstance()
+	{
+		return MakeShared<FMediaStreamSourceCustomization>();
+	}
+
+	FMediaStreamSourceCustomization();
+	virtual ~FMediaStreamSourceCustomization() override = default;
+
+	//~ Begin IPropertyTypeCustomization
+	virtual void CustomizeHeader(TSharedRef<IPropertyHandle> InPropertyHandle, FDetailWidgetRow& InHeaderRow, IPropertyTypeCustomizationUtils& InCustomizationUtils) override;
+	virtual void CustomizeChildren(TSharedRef<IPropertyHandle> InPropertyHandle, IDetailChildrenBuilder& InChildBuilder, IPropertyTypeCustomizationUtils& InCustomizationUtils) override;
+	//~ End IPropertyTypeCustomization
+
+private:
+	TSharedPtr<IPropertyHandle> PropertyHandle;
+
+	void AddSourceSchemeSelector(IDetailGroup& InSourceGroup, IPropertyTypeCustomizationUtils& InCustomizationUtils);
+
+	void AddSourceSchemeSelectRow(TSharedRef<SGridPanel> InContainer, int32 InRow, const FName& InName);
+
+	ECheckBoxState GetSourceCheckBoxState(FName InSourceScheme) const;
+	void OnSourceCheckBoxStateChanged(ECheckBoxState InState, FName InSourceScheme);
+
+	FReply OnSourceSchemeButtonClicked(const FName& InSourceScheme);
+
+	void AddSchemeCustomizations(TSharedRef<IPropertyHandle> InPropertyHandle, IDetailGroup& InSourceGroup, IPropertyTypeCustomizationUtils& InCustomizationUtils);
+};
diff --git a/Engine/Plugins/Experimental/MediaStream/Source/MediaStreamEditor/Private/MediaStreamEditorModule.cpp b/Engine/Plugins/Experimental/MediaStream/Source/MediaStreamEditor/Private/MediaStreamEditorModule.cpp
new file mode 100644
--- /dev/null
+++ b/Engine/Plugins/Experimental/MediaStream/Source/MediaStreamEditor/Private/MediaStreamEditorModule.cpp
@@ -0,0 +1,34 @@
+// Copyright Epic Games, Inc. All Rights Reserved.
+
+#include "MediaStreamEditorModule.h"
+
+#include "DetailsPanel/MediaStreamCustomization.h"
+#include "DetailsPanel/MediaStreamSourceCustomization.h"
+#include "MediaStream.h"
+#include "MediaStreamEditorStyle.h"
+#include "MediaStreamEditorToolbar.h"
+#include "Modules/ModuleManager.h"
+#include "PropertyEditorDelegates.h"
+#include "PropertyEditorModule.h"
+
+void FMediaStreamEditorModule::StartupModule()
+{
+	FMediaStreamEditorStyle::Get();
+
+	FPropertyEditorModule& PropertyModule = FModuleManager::LoadModuleChecked<FPropertyEditorModule>("PropertyEditor");
+
+	PropertyModule.RegisterCustomClassLayout(UMediaStream::StaticClass()->GetFName(),
+		FOnGetDetailCustomizationInstance::CreateStatic(&FMediaStreamCustomization::MakeInstance));
+
+	PropertyModule.RegisterCustomPropertyTypeLayout(FMediaStreamSource::StaticStruct()->GetFName(),
+		FOnGetPropertyTypeCustomizationInstance::CreateStatic(&FMediaStreamSourceCustomization::MakeInstance));
+
+	FMediaStreamEditorToolbar::Get().Integrate();
+}
+
+void FMediaStreamEditorModule::ShutdownModule()
+{
+	FMediaStreamEditorToolbar::Get().Disintegrate();
+}
+
+IMPLEMENT_MODULE(FMediaStreamEditorModule, MediaStreamEditor);
diff --git a/Engine/Plugins/Experimental/MediaStream/Source/MediaStreamEditor/Private/MediaStreamEditorModule.h b/Engine/Plugins/Experimental/MediaStream/Source/MediaStreamEditor/Private/MediaStreamEditorModule.h
new file mode 100644
--- /dev/null
+++ b/Engine/Plugins/Experimental/MediaStream/Source/MediaStreamEditor/Private/MediaStreamEditorModule.h
@@ -0,0 +1,19 @@
+// Copyright Epic Games, Inc. All Rights Reserved.
+
+#pragma once
+
+#include "Modules/ModuleInterface.h"
+
+#include "Delegates/IDelegateInstance.h"
+
+/**
+ * Media Stream Editor - Content/type agnostic chainable media proxy with media player integration.
+ */
+class FMediaStreamEditorModule : public IModuleInterface
+{
+public:
+	//~ Begin IModuleInterface
+	virtual void StartupModule() override;
+	virtual void ShutdownModule() override;
+	//~ End IModuleInterface
+};
diff --git a/Engine/Plugins/Experimental/MediaStream/Source/MediaStreamEditor/Private/MediaStreamEditorSequencerLibrary.cpp b/Engine/Plugins/Experimental/MediaStream/Source/MediaStreamEditor/Private/MediaStreamEditorSequencerLibrary.cpp
new file mode 100644
--- /dev/null
+++ b/Engine/Plugins/Experimental/MediaStream/Source/MediaStreamEditor/Private/MediaStreamEditorSequencerLibrary.cpp
@@ -0,0 +1,223 @@
+// Copyright Epic Games, Inc. All Rights Reserved.
+
+#include "MediaStreamEditorSequencerLibrary.h"
+
+#include "IMediaStreamPlayer.h"
+#include "LevelSequence.h"
+#include "LevelSequenceEditorBlueprintLibrary.h"
+#include "MediaStream.h"
+#include "MovieScene.h"
+#include "MovieSceneCommonHelpers.h"
+#include "MovieSceneMediaSection.h"
+#include "MovieSceneMediaTrack.h"
+
+ULevelSequence* FMediaStreamEditorSequencerLibrary::GetLevelSequence()
+{
+	if (ULevelSequence* FocusedSequence = ULevelSequenceEditorBlueprintLibrary::GetFocusedLevelSequence())
+	{
+		return FocusedSequence;
+	}
+
+	if (ULevelSequence* FocusedSequence = ULevelSequenceEditorBlueprintLibrary::GetCurrentLevelSequence())
+	{
+		return FocusedSequence;
+	}
+
+	return nullptr;
+}
+
+bool FMediaStreamEditorSequencerLibrary::HasTrack(UMediaStream* InMediaStream)
+{
+	if (!InMediaStream)
+	{
+		return false;
+	}
+
+	ULevelSequence* LevelSequence = GetLevelSequence();
+
+	if (!LevelSequence)
+	{
+		return false;
+	}
+
+	UMovieScene* MovieScene = LevelSequence->GetMovieScene();
+
+	if (!MovieScene)
+	{
+		return false;
+	}
+
+	UWorld* World = LevelSequence->GetWorld();
+
+	if (!World)
+	{
+		return false;
+	}
+
+	TSharedRef<const UE::MovieScene::FSharedPlaybackState> SharedPlaybackState = MovieSceneHelpers::CreateTransientSharedPlaybackState(World, LevelSequence);
+	const FGuid MediaStreamId = LevelSequence->FindBindingFromObject(InMediaStream, SharedPlaybackState);
+
+	if (MediaStreamId.IsValid())
+	{
+		return !!MovieScene->FindTrack<UMovieSceneMediaTrack>(MediaStreamId);
+	}
+
+	return false;
+}
+
+bool FMediaStreamEditorSequencerLibrary::CanAddTrack(UMediaStream* InMediaStream)
+{
+	if (!InMediaStream)
+	{
+		return false;
+	}
+
+	UObject* MediaObject = InMediaStream->GetSource().Object;
+
+	// Can't add a proxy stream to a track.
+	if (!IsValid(MediaObject) || MediaObject->IsA<UMediaStream>())
+	{
+		return false;
+	}
+
+	IMediaStreamPlayer* MediaStreamPlayerIntf = InMediaStream->GetPlayer().GetInterface();
+
+	if (!MediaStreamPlayerIntf || MediaStreamPlayerIntf->IsReadOnly())
+	{
+		return false;
+	}
+
+	UMediaPlayer* MediaPlayer = MediaStreamPlayerIntf->GetPlayer();
+
+	if (!MediaPlayer)
+	{
+		return false;
+	}
+
+	ULevelSequence* LevelSequence = GetLevelSequence();
+
+	if (!LevelSequence)
+	{
+		return false;
+	}
+
+	UMovieScene* MovieScene = LevelSequence->GetMovieScene();
+
+	if (!MovieScene)
+	{
+		return false;
+	}
+
+	UWorld* World = LevelSequence->GetWorld();
+
+	if (!World)
+	{
+		return false;
+	}
+
+	TSharedRef<const UE::MovieScene::FSharedPlaybackState> SharedPlaybackState = MovieSceneHelpers::CreateTransientSharedPlaybackState(World, LevelSequence);
+	const FGuid MediaStreamId = LevelSequence->FindBindingFromObject(InMediaStream, SharedPlaybackState);
+
+	if (MediaStreamId.IsValid())
+	{
+		return !MovieScene->FindTrack<UMovieSceneMediaTrack>(MediaStreamId);
+	}
+
+	return LevelSequence->CanPossessObject(*InMediaStream, World);
+}
+
+void FMediaStreamEditorSequencerLibrary::AddTrack(UMediaStream* InMediaStream)
+{
+	if (!InMediaStream || !CanAddTrack(InMediaStream))
+	{
+		return;
+	}
+
+	IMediaStreamPlayer* MediaStreamPlayerIntf = InMediaStream->GetPlayer().GetInterface();
+
+	if (!MediaStreamPlayerIntf || MediaStreamPlayerIntf->IsReadOnly())
+	{
+		return;
+	}
+
+	UMediaPlayer* MediaPlayer = MediaStreamPlayerIntf->GetPlayer();
+
+	if (!MediaPlayer)
+	{
+		return;
+	}
+
+	ULevelSequence* LevelSequence = GetLevelSequence();
+
+	if (!LevelSequence)
+	{
+		return;
+	}
+
+	UMovieScene* MovieScene = LevelSequence->GetMovieScene();
+
+	if (!MovieScene)
+	{
+		return;
+	}
+
+	UWorld* World = LevelSequence->GetWorld();
+
+	if (!World)
+	{
+		return;
+	}
+
+	TSharedRef<const UE::MovieScene::FSharedPlaybackState> SharedPlaybackState = MovieSceneHelpers::CreateTransientSharedPlaybackState(World, LevelSequence);
+	FGuid MediaStreamId = LevelSequence->FindBindingFromObject(InMediaStream, SharedPlaybackState);
+
+	const FString MediaPlayerPath = MediaPlayer->GetPathName();
+
+	if (!MediaStreamId.IsValid())
+	{
+		const FString DefaultBindingName = TEXT("MediaStream");
+		const FMovieScenePossessable Possessable(DefaultBindingName, UMediaStream::StaticClass());
+		const FMovieSceneBinding Binding(Possessable.GetGuid(), DefaultBindingName);
+		MovieScene->AddPossessable(Possessable, Binding);
+		MediaStreamId = Possessable.GetGuid();
+		LevelSequence->BindPossessableObject(MediaStreamId, *InMediaStream, World);
+
+		if (!MediaStreamId.IsValid())
+		{
+			return;
+		}
+	}
+
+	if (MovieScene->FindTrack<UMovieSceneMediaTrack>(MediaStreamId))
+	{
+		return;
+	}
+
+	UMovieSceneMediaTrack* MediaTrack = MovieScene->AddTrack<UMovieSceneMediaTrack>(MediaStreamId);
+
+	if (!MediaTrack)
+	{
+		return;
+	}
+
+	TRange<FFrameNumber> Range = MovieScene->GetPlaybackRange();
+
+	if (Range.GetLowerBound().IsOpen())
+	{
+		Range.SetLowerBound(FFrameNumber(0));
+	}
+
+	if (Range.GetUpperBound().IsOpen() || Range.GetUpperBound().GetValue() == 0)
+	{
+		Range.SetUpperBound(FFrameNumber(250));
+	}
+
+	UMovieSceneMediaSection* MediaSection = Cast<UMovieSceneMediaSection>(MediaTrack->CreateNewSection());
+	MediaSection->bUseExternalMediaPlayer = true;
+	MediaSection->ExternalMediaPlayer = MediaPlayer;
+	MediaSection->SetStartFrame(Range.GetLowerBound());
+	MediaSection->SetEndFrame(Range.GetUpperBound());
+	MediaSection->SetMediaSourceProxy(UE::MovieScene::FRelativeObjectBindingID(MediaStreamId), /* Proxy Index */ 0);
+
+	MediaTrack->AddSection(*MediaSection);
+}
diff --git a/Engine/Plugins/Experimental/MediaStream/Source/MediaStreamEditor/Private/MediaStreamEditorSequencerLibrary.h b/Engine/Plugins/Experimental/MediaStream/Source/MediaStreamEditor/Private/MediaStreamEditorSequencerLibrary.h
new file mode 100644
--- /dev/null
+++ b/Engine/Plugins/Experimental/MediaStream/Source/MediaStreamEditor/Private/MediaStreamEditorSequencerLibrary.h
@@ -0,0 +1,19 @@
+// Copyright Epic Games, Inc. All Rights Reserved.
+
+#pragma once
+
+class UObject;
+class ULevelSequence;
+class UMediaStream;
+
+class FMediaStreamEditorSequencerLibrary
+{
+public:
+	static ULevelSequence* GetLevelSequence();
+
+	static bool HasTrack(UMediaStream* InMediaStream);
+
+	static bool CanAddTrack(UMediaStream* InMediaStream);
+
+	static void AddTrack(UMediaStream* InMediaStream);
+};
diff --git a/Engine/Plugins/Experimental/MediaStream/Source/MediaStreamEditor/Private/MediaStreamEditorStyle.cpp b/Engine/Plugins/Experimental/MediaStream/Source/MediaStreamEditor/Private/MediaStreamEditorStyle.cpp
new file mode 100644
--- /dev/null
+++ b/Engine/Plugins/Experimental/MediaStream/Source/MediaStreamEditor/Private/MediaStreamEditorStyle.cpp
@@ -0,0 +1,101 @@
+// Copyright Epic Games, Inc. All Rights Reserved.
+
+#include "MediaStreamEditorStyle.h"
+
+#include "Brushes/SlateBorderBrush.h"
+#include "Brushes/SlateBoxBrush.h"
+#include "Brushes/SlateImageBrush.h"
+#include "Fonts/SlateFontInfo.h"
+#include "Interfaces/IPluginManager.h"
+#include "Math/Vector2D.h"
+#include "Misc/Paths.h"
+#include "Styling/SlateStyleMacros.h"
+#include "Styling/SlateStyleRegistry.h"
+#include "Styling/SlateTypes.h"
+#include "Styling/CoreStyle.h"
+
+/* FMediaStreamEditorStyle structors
+ *****************************************************************************/
+
+FMediaStreamEditorStyle::FMediaStreamEditorStyle()
+	: FSlateStyleSet("MediaStreamEditorStyle")
+{
+	const FVector2D Icon8x8(8.0f, 8.0f);
+	const FVector2D Icon12x12(12.0f, 12.0f);
+	const FVector2D Icon16x16(16.0f, 16.0f);
+	const FVector2D Icon20x20(20.0f, 20.0f);
+	const FVector2D Icon40x40(40.0f, 40.0f);
+	const FVector2D Icon64x64(64.0f, 64.0f);
+
+	if (TSharedPtr<IPlugin> Plugin = IPluginManager::Get().FindPlugin(TEXT("MediaStream")))
+	{
+		SetContentRoot(FPaths::Combine(Plugin->GetBaseDir(), TEXT("Resources")));
+	}
+	// Class icons.
+	Set("ClassIcon.MediaStream", new IMAGE_BRUSH_SVG(TEXT("Icons/MediaStream"), Icon16x16));
+	Set("ClassThumbnail.MediaStream", new IMAGE_BRUSH_SVG(TEXT("Icons/MediaStream"), Icon64x64));
+	Set("ClassIcon.MediaStreamActor", new IMAGE_BRUSH_SVG(TEXT("Icons/MediaStream"), Icon16x16));
+	Set("ClassThumbnail.MediaStreamActor", new IMAGE_BRUSH_SVG(TEXT("Icons/MediaStream"), Icon64x64));
+	Set("ClassIcon.MediaStreamComponent", new IMAGE_BRUSH_SVG(TEXT("Icons/MediaStream"), Icon16x16));
+	Set("ClassThumbnail.MediaStreamComponent", new IMAGE_BRUSH_SVG(TEXT("Icons/MediaStream"), Icon64x64));
+
+	SetContentRoot(FPaths::EnginePluginsDir() / TEXT("Media/MediaPlayerEditor/Content"));
+
+	// buttons
+	Set("MediaPlayerEditor.SourceButton", new IMAGE_BRUSH("btn_source_12x", Icon12x12));
+	Set("MediaPlayerEditor.GoButton", new IMAGE_BRUSH("btn_go_12x", Icon12x12));
+	Set("MediaPlayerEditor.ReloadButton", new IMAGE_BRUSH("btn_reload_12x", Icon12x12));
+	Set("MediaPlayerEditor.SettingsButton", new IMAGE_BRUSH("btn_settings_16x", Icon12x12));
+
+	// misc
+	Set("MediaStreamEditor.DragDropBorder", new BOX_BRUSH("border_dragdrop", 0.5f));
+	Set("MediaStreamEditor.MediaSourceOpened", new IMAGE_BRUSH("mediasource_opened", Icon8x8));
+
+	// tabs
+	Set("MediaStreamEditor.Tabs.Info", new IMAGE_BRUSH("tab_info_16x", Icon16x16));
+	Set("MediaStreamEditor.Tabs.Media", new IMAGE_BRUSH("tab_media_16x", Icon16x16));
+	Set("MediaStreamEditor.Tabs.Player", new IMAGE_BRUSH("tab_player_16x", Icon16x16));
+	Set("MediaStreamEditor.Tabs.Playlist", new IMAGE_BRUSH("tab_playlist_16x", Icon16x16));
+	Set("MediaStreamEditor.Tabs.Stats", new IMAGE_BRUSH("tab_stats_16x", Icon16x16));
+
+	// toolbar icons
+	Set("MediaStreamEditor.CloseMedia", new IMAGE_BRUSH("icon_eject_40x", Icon40x40));
+	Set("MediaStreamEditor.CloseMedia.Small", new IMAGE_BRUSH("icon_eject_40x", Icon20x20));
+	Set("MediaStreamEditor.ForwardMedia", new IMAGE_BRUSH("icon_rewind_40x", Icon40x40));
+	Set("MediaStreamEditor.ForwardMedia.Small", new IMAGE_BRUSH("icon_rewind_40x", Icon20x20));
+	Set("MediaStreamEditor.NextMedia", new IMAGE_BRUSH("icon_step_40x", Icon40x40));
+	Set("MediaStreamEditor.NextMedia.Small", new IMAGE_BRUSH("icon_step_40x", Icon20x20));
+	Set("MediaStreamEditor.OpenMedia", new IMAGE_BRUSH("icon_open_40x", Icon40x40));
+	Set("MediaStreamEditor.OpenMedia.Small", new IMAGE_BRUSH("icon_open_40x", Icon20x20));
+	Set("MediaStreamEditor.PauseMedia", new IMAGE_BRUSH("icon_pause_40x", Icon40x40));
+	Set("MediaStreamEditor.PauseMedia.Small", new IMAGE_BRUSH("icon_pause_40x", Icon20x20));
+	Set("MediaStreamEditor.PlayMedia", new IMAGE_BRUSH("icon_play_40x", Icon40x40));
+	Set("MediaStreamEditor.PlayMedia.Small", new IMAGE_BRUSH("icon_play_40x", Icon20x20));
+	Set("MediaStreamEditor.PreviousMedia", new IMAGE_BRUSH("icon_step_back_40x", Icon40x40));
+	Set("MediaStreamEditor.PreviousMedia.Small", new IMAGE_BRUSH("icon_step_back_40x", Icon20x20));
+	Set("MediaStreamEditor.ReverseMedia", new IMAGE_BRUSH("icon_reverse_40x", Icon40x40));
+	Set("MediaStreamEditor.ReverseMedia.Small", new IMAGE_BRUSH("icon_reverse_40x", Icon20x20));
+	Set("MediaStreamEditor.RewindMedia", new IMAGE_BRUSH("icon_rewind_40x", Icon40x40));
+	Set("MediaStreamEditor.RewindMedia.Small", new IMAGE_BRUSH("icon_rewind_40x", Icon20x20));
+	Set("MediaStreamEditor.StopMedia", new IMAGE_BRUSH("icon_stop_40x", Icon40x40));
+	Set("MediaStreamEditor.StopMedia.Small", new IMAGE_BRUSH("icon_stop_40x", Icon20x20));
+
+	// scrubber
+	Set("MediaPlayerEditor.Scrubber", FSliderStyle()
+		.SetNormalBarImage(FSlateColorBrush(FColor::White))
+		.SetDisabledBarImage(FSlateColorBrush(FLinearColor::Gray))
+		.SetNormalThumbImage(IMAGE_BRUSH("scrubber", FVector2D(2.0f, 10.0f)))
+		.SetHoveredThumbImage(IMAGE_BRUSH("scrubber", FVector2D(2.0f, 10.0f)))
+		.SetDisabledThumbImage(IMAGE_BRUSH("scrubber", FVector2D(2.0f, 10.0f)))
+		.SetBarThickness(2.0f)
+	);
+
+	Set("MediaStreamEditor.ViewportFont", DEFAULT_FONT("Regular", 18));
+
+	FSlateStyleRegistry::RegisterSlateStyle(*this);
+}
+
+FMediaStreamEditorStyle::~FMediaStreamEditorStyle()
+{
+	FSlateStyleRegistry::UnRegisterSlateStyle(*this);
+}
diff --git a/Engine/Plugins/Experimental/MediaStream/Source/MediaStreamEditor/Private/MediaStreamEditorStyle.h b/Engine/Plugins/Experimental/MediaStream/Source/MediaStreamEditor/Private/MediaStreamEditorStyle.h
new file mode 100644
--- /dev/null
+++ b/Engine/Plugins/Experimental/MediaStream/Source/MediaStreamEditor/Private/MediaStreamEditorStyle.h
@@ -0,0 +1,30 @@
+// Copyright Epic Games, Inc. All Rights Reserved.
+
+#pragma once
+
+#include "Styling/SlateStyle.h"
+
+/**
+ * Implements the visual style of the media plate editor UI.
+ */
+class FMediaStreamEditorStyle : public FSlateStyleSet
+{
+public:
+	/** Default constructor. */
+	FMediaStreamEditorStyle();
+
+	 /** Destructor. */
+	~FMediaStreamEditorStyle();
+
+	static TSharedRef<FMediaStreamEditorStyle> Get()
+	{
+		static TSharedPtr<FMediaStreamEditorStyle> Singleton;
+
+		if (!Singleton.IsValid())
+		{
+			Singleton = MakeShareable(new FMediaStreamEditorStyle);
+		}
+
+		return Singleton.ToSharedRef();
+	}
+};
diff --git a/Engine/Plugins/Experimental/MediaStream/Source/MediaStreamEditor/Private/MediaStreamEditorToolbar.cpp b/Engine/Plugins/Experimental/MediaStream/Source/MediaStreamEditor/Private/MediaStreamEditorToolbar.cpp
new file mode 100644
--- /dev/null
+++ b/Engine/Plugins/Experimental/MediaStream/Source/MediaStreamEditor/Private/MediaStreamEditorToolbar.cpp
@@ -0,0 +1,126 @@
+// Copyright Epic Games, Inc. All Rights Reserved.
+
+#include "MediaStreamEditorToolbar.h"
+
+#include "MediaStreamEditorStyle.h"
+#include "MediaStreamModule.h"
+#include "Misc/Attribute.h"
+#include "Textures/SlateIcon.h"
+#include "ToolMenu.h"
+#include "ToolMenus.h"
+#include "ToolMenuSection.h"
+
+#define LOCTEXT_NAMESPACE "MediaStreamEditorToolbar"
+
+namespace UE::MediaStreamEditor::Private
+{
+	static constexpr const TCHAR* LevelEditorToolbarName = TEXT("LevelEditor.LevelEditorToolBar.User");
+	static constexpr const TCHAR* LevelEditorToolbarSectionName = TEXT("MediaStream");
+}
+
+FMediaStreamEditorToolbar& FMediaStreamEditorToolbar::Get()
+{
+	static FMediaStreamEditorToolbar Instance;
+	return Instance;
+}
+
+void FMediaStreamEditorToolbar::Integrate()
+{
+	ExtendLevelEditorToolbar();
+}
+
+void FMediaStreamEditorToolbar::Disintegrate()
+{
+}
+
+void FMediaStreamEditorToolbar::ExtendLevelEditorToolbar()
+{
+	using namespace UE::MediaStreamEditor::Private;
+
+	UToolMenus* const ToolMenus = UToolMenus::Get();
+
+	if (!ToolMenus)
+	{
+		return;
+	}
+
+	UToolMenu* const LevelEditorToolbar = ToolMenus->ExtendMenu(LevelEditorToolbarName);
+
+	if (!LevelEditorToolbar)
+	{
+		return;
+	}
+
+	FToolMenuSection& Section = LevelEditorToolbar->FindOrAddSection(LevelEditorToolbarSectionName);
+
+	FToolUIAction Action;
+	Action.ExecuteAction.BindRaw(this, &FMediaStreamEditorToolbar::OnToolbarButtonClicked);
+
+	Section.AddEntry(FToolMenuEntry::InitToolBarButton(
+		TEXT("ToggleMediaStream"),
+		FToolUIActionChoice(Action),
+		FText::GetEmpty(),
+		LOCTEXT("MediaStreamToolTip", "Start or stop all Media Streams."),
+		TAttribute<FSlateIcon>::CreateRaw(this, &FMediaStreamEditorToolbar::GetIcon)
+	));
+}
+
+void FMediaStreamEditorToolbar::UnextendLevelEditorToolbar()
+{
+	using namespace UE::MediaStreamEditor::Private;
+
+	UToolMenus* const ToolMenus = UToolMenus::Get();
+
+	if (!ToolMenus)
+	{
+		return;
+	}
+
+	UToolMenu* const LevelEditorToolbar = ToolMenus->ExtendMenu(LevelEditorToolbarName);
+
+	if (!LevelEditorToolbar)
+	{
+		return;
+	}
+
+	LevelEditorToolbar->RemoveSection(LevelEditorToolbarSectionName);
+}
+
+void FMediaStreamEditorToolbar::OnToolbarButtonClicked(const FToolMenuContext& InContext)
+{
+	switch (EditorPlaybackState)
+	{
+		case EMediaStreamPlaybackState::Pause:
+			EditorPlaybackState = EMediaStreamPlaybackState::Play;
+			break;
+
+		case EMediaStreamPlaybackState::Play:
+			EditorPlaybackState = EMediaStreamPlaybackState::Pause;
+			break;
+
+		default:
+			return;
+	}
+
+	FMediaStreamModule::Get().GetEditorPlaybackStateEvent().Broadcast(EditorPlaybackState);
+}
+
+FSlateIcon FMediaStreamEditorToolbar::GetIcon() const
+{
+	static const FSlateIcon PlayButton = FSlateIcon(FMediaStreamEditorStyle::Get()->GetStyleSetName(), TEXT("MediaStreamEditor.PlayMedia"));
+	static const FSlateIcon PauseButton = FSlateIcon(FMediaStreamEditorStyle::Get()->GetStyleSetName(), TEXT("MediaStreamEditor.PauseMedia"));
+
+	switch (EditorPlaybackState)
+	{
+		case EMediaStreamPlaybackState::Pause:
+			return PlayButton;
+
+		case EMediaStreamPlaybackState::Play:
+			return PauseButton;
+
+		default:
+			return FSlateIcon();
+	}
+}
+
+#undef LOCTEXT_NAMESPACE
diff --git a/Engine/Plugins/Experimental/MediaStream/Source/MediaStreamEditor/Private/MediaStreamEditorToolbar.h b/Engine/Plugins/Experimental/MediaStream/Source/MediaStreamEditor/Private/MediaStreamEditorToolbar.h
new file mode 100644
--- /dev/null
+++ b/Engine/Plugins/Experimental/MediaStream/Source/MediaStreamEditor/Private/MediaStreamEditorToolbar.h
@@ -0,0 +1,30 @@
+// Copyright Epic Games, Inc. All Rights Reserved.
+
+#pragma once
+
+#include "MediaStreamEnums.h"
+
+struct FToolMenuContext;
+struct FSlateIcon;
+
+/** Responsible for manipulating the Level Editor toolbar for the Motion Design Editor */
+class FMediaStreamEditorToolbar
+{
+public:
+	static FMediaStreamEditorToolbar& Get();
+
+	void Integrate();
+
+	void Disintegrate();
+
+private:
+	EMediaStreamPlaybackState EditorPlaybackState = EMediaStreamPlaybackState::Pause;
+
+	void ExtendLevelEditorToolbar();
+
+	void UnextendLevelEditorToolbar();
+
+	void OnToolbarButtonClicked(const FToolMenuContext& InContext);
+
+	FSlateIcon GetIcon() const;
+};
