Media Stream: Load guard
- Added a guard to stop sources automatically loading in the editor until user action is taken.
- Does not affect -game or packaged.

diff --git a/Engine/Plugins/Experimental/MediaStream/Source/MediaStream/Private/MediaStream.cpp b/Engine/Plugins/Experimental/MediaStream/Source/MediaStream/Private/MediaStream.cpp
--- a/Engine/Plugins/Experimental/MediaStream/Source/MediaStream/Private/MediaStream.cpp
+++ b/Engine/Plugins/Experimental/MediaStream/Source/MediaStream/Private/MediaStream.cpp
@@ -218,6 +218,9 @@ bool UMediaStream::ApplySource()
 			if (IMediaStreamPlayer* Player = Cast<IMediaStreamPlayer>(PlayerObject))
 			{
 				Player->Deinitialize();
+				Player->OnSourceChanged(Source);
+
+				OnPlayerChanged.Broadcast(this, PlayerObject);
 			}
 		}
 
diff --git a/Engine/Plugins/Experimental/MediaStream/Source/MediaStream/Private/MediaStreamModule.cpp b/Engine/Plugins/Experimental/MediaStream/Source/MediaStream/Private/MediaStreamModule.cpp
--- a/Engine/Plugins/Experimental/MediaStream/Source/MediaStream/Private/MediaStreamModule.cpp
+++ b/Engine/Plugins/Experimental/MediaStream/Source/MediaStream/Private/MediaStreamModule.cpp
@@ -21,9 +21,18 @@ FMediaStreamModule& FMediaStreamModule::Get()
 	return FModuleManager::Get().GetModuleChecked<FMediaStreamModule>(UE_MODULE_NAME);
 }
 
-FMediaStreamEditorPlaybackStateDelegate& FMediaStreamModule::GetEditorPlaybackStateEvent()
+bool FMediaStreamModule::CanOpenSourceOnLoad()
 {
-	return EditorPlaybackStateDelegate;
+#if WITH_EDITOR
+	return !GIsEditor;
+#else
+	return true;
+#endif
+}
+
+bool FMediaStreamModule::CanAutoplay() const
+{
+	return !GIsEditor;
 }
 
 void FMediaStreamModule::StartupModule()
diff --git a/Engine/Plugins/Experimental/MediaStream/Source/MediaStream/Private/MediaStreamObjectHandlerManager.cpp b/Engine/Plugins/Experimental/MediaStream/Source/MediaStream/Private/MediaStreamObjectHandlerManager.cpp
--- a/Engine/Plugins/Experimental/MediaStream/Source/MediaStream/Private/MediaStreamObjectHandlerManager.cpp
+++ b/Engine/Plugins/Experimental/MediaStream/Source/MediaStream/Private/MediaStreamObjectHandlerManager.cpp
@@ -40,66 +40,83 @@ bool FMediaStreamObjectHandlerManager::CanHandleObject(const UObject* InObject)
 	return CanHandleObject(InObject->GetClass());
 }
 
-UMediaPlayer* FMediaStreamObjectHandlerManager::CreateOrUpdateMediaPlayer(UMediaStream* InMediaStream, UObject* InSourceObject, UMediaPlayer* InCurrentPlayer) const
+UMediaPlayer* FMediaStreamObjectHandlerManager::CreateOrUpdatePlayer(const FMediaStreamObjectHandlerCreatePlayerParams& InParams) const
 {
-	if (!IsValid(InMediaStream))
+	if (!IsValid(InParams.MediaStream))
 	{
-		UE_LOG(LogMediaStream, Error, TEXT("Invalid Media Stream in FUMediaStreamObjectHandlerManager::CreateOrUpdateMediaPlayer"));
+		UE_LOG(LogMediaStream, Error, TEXT("Invalid Media Stream in FMediaStreamObjectHandlerManager::CreateOrUpdatePlayer"));
 		return nullptr;
 	}
 
-	if (!IsValid(InSourceObject))
+	if (!IsValid(InParams.Source))
 	{
-		UE_LOG(LogMediaStream, Error, TEXT("Invalid Object in FUMediaStreamObjectHandlerManager::CreateOrUpdateMediaPlayer"));
+		UE_LOG(LogMediaStream, Error, TEXT("Invalid Object in FMediaStreamObjectHandlerManager::CreateOrUpdatePlayer"));
 		return nullptr;
 	}
 
-	if (InCurrentPlayer)
+	if (InParams.CurrentPlayer)
 	{
-		if (UMediaSource* MediaSource = Cast<UMediaSource>(InSourceObject))
+		if (UMediaSource* MediaSource = Cast<UMediaSource>(InParams.Source))
 		{
-			if (InCurrentPlayer->OpenSource(MediaSource))
+			const bool bIsValidPlayer = InParams.bCanOpenSource
+				? InParams.CurrentPlayer->OpenSource(MediaSource)
+				: InParams.CurrentPlayer->CanPlaySource(MediaSource);
+
+			if (bIsValidPlayer)
 			{
-				return InCurrentPlayer;
+				return InParams.CurrentPlayer;
 			}
 		}
-		else if (UMediaPlaylist* Playlist = Cast<UMediaPlaylist>(InSourceObject))
+		else if (UMediaPlaylist* Playlist = Cast<UMediaPlaylist>(InParams.Source))
 		{
-			if (InCurrentPlayer->OpenPlaylist(Playlist))
+			UMediaSource* FirstPlaylistEntry = Playlist->Get(0);
+
+			const bool bIsValidPlayer = InParams.bCanOpenSource
+				? InParams.CurrentPlayer->OpenPlaylist(Playlist)
+				: (!FirstPlaylistEntry || InParams.CurrentPlayer->CanPlaySource(FirstPlaylistEntry));
+
+			if (bIsValidPlayer)
 			{
-				return InCurrentPlayer;
+				return InParams.CurrentPlayer;
 			}
 		}
 	}
 
-	if (TSharedPtr<IMediaStreamObjectHandler> Handler = GetHandler(InSourceObject->GetClass()))
+	if (TSharedPtr<IMediaStreamObjectHandler> Handler = GetHandler(InParams.Source->GetClass()))
 	{
-		UMediaPlayer* MediaPlayer = Handler->CreateOrUpdatePlayer(InMediaStream, InSourceObject, InCurrentPlayer);
+		UMediaPlayer* MediaPlayer = Handler->CreateOrUpdatePlayer(InParams);
 
 		if (!MediaPlayer)
 		{
-			UE_LOG(LogMediaStream, Error, TEXT("Failed to create Media Player in FUMediaStreamObjectHandlerManager::CreateOrUpdateMediaPlayer [%s] [%s]"), *InSourceObject->GetClass()->GetName(), *InSourceObject->GetName());
+			UE_LOG(LogMediaStream, Error, TEXT("Failed to create Media Player in FMediaStreamObjectHandlerManager::CreateOrUpdatePlayer [%s] [%s]"), *InParams.Source->GetClass()->GetName(), *InParams.Source->GetName());
 			return nullptr;
 		}
 
 		return MediaPlayer;
 	}
 
+	// We cannot allow players to open sources, which we cannot control beyond this point.
+	if (!InParams.bCanOpenSource)
+	{
+		UE_LOG(LogMediaStream, Error, TEXT("Cannot use Media Assets Module to create player at the moment in FMediaStreamObjectHandlerManager::CreateOrUpdatePlayer [%s] [%s]"), *InParams.Source->GetClass()->GetName(), *InParams.Source->GetName());
+		return nullptr;
+	}
+
 	if (IMediaAssetsModule* MediaAssetsModule = FModuleManager::LoadModulePtr<IMediaAssetsModule>("MediaAssets"))
 	{
 		UObject* PlayerProxy = nullptr;
-		UMediaPlayer* MediaPlayer = MediaAssetsModule->GetPlayerFromObject(InSourceObject, PlayerProxy);
+		UMediaPlayer* MediaPlayer = MediaAssetsModule->GetPlayerFromObject(InParams.Source, PlayerProxy);
 
 		if (!MediaPlayer)
 		{
-			UE_LOG(LogMediaStream, Error, TEXT("No player for Object in FUMediaStreamObjectHandlerManager::CreateOrUpdateMediaPlayer [%s] [%s]"), *InSourceObject->GetClass()->GetName(), *InSourceObject->GetName());
+			UE_LOG(LogMediaStream, Error, TEXT("No player for Object in FMediaStreamObjectHandlerManager::CreateOrUpdatePlayer [%s] [%s]"), *InParams.Source->GetClass()->GetName(), *InParams.Source->GetName());
 			return nullptr;
 		}
 
 		return MediaPlayer;
 	}
 
-	UE_LOG(LogMediaStream, Error, TEXT("Unable to get Media Assets Module in FUMediaStreamObjectHandlerManager::CreateOrUpdateMediaPlayer [%s] [%s]"), *InSourceObject->GetClass()->GetName(), *InSourceObject->GetName());
+	UE_LOG(LogMediaStream, Error, TEXT("Unable to get Media Assets Module in FMediaStreamObjectHandlerManager::CreateOrUpdatePlayer [%s] [%s]"), *InParams.Source->GetClass()->GetName(), *InParams.Source->GetName());
 	return nullptr;
 }
 
diff --git a/Engine/Plugins/Experimental/MediaStream/Source/MediaStream/Private/MediaStreamObjectHandlerSubsystem.cpp b/Engine/Plugins/Experimental/MediaStream/Source/MediaStream/Private/MediaStreamObjectHandlerSubsystem.cpp
--- a/Engine/Plugins/Experimental/MediaStream/Source/MediaStream/Private/MediaStreamObjectHandlerSubsystem.cpp
+++ b/Engine/Plugins/Experimental/MediaStream/Source/MediaStream/Private/MediaStreamObjectHandlerSubsystem.cpp
@@ -36,9 +36,9 @@ bool UMediaStreamObjectHandlerSubsystem::CanHandleObject(const UClass* InClass)
 	return FMediaStreamObjectHandlerManager::Get().CanHandleObject(InClass);
 }
 
-UMediaPlayer* UMediaStreamObjectHandlerSubsystem::CreateMediaPlayer(UMediaStream* InMediaStream, UObject* InPlaybackSource, UMediaPlayer* InCurrentPlayer) const
+UMediaPlayer* UMediaStreamObjectHandlerSubsystem::CreateMediaPlayer(const FMediaStreamObjectHandlerCreatePlayerParams& InParams) const
 {
-	return FMediaStreamObjectHandlerManager::Get().CreateOrUpdateMediaPlayer(InMediaStream, InPlaybackSource, InCurrentPlayer);
+	return FMediaStreamObjectHandlerManager::Get().CreateOrUpdatePlayer(InParams);
 }
 
 bool UMediaStreamObjectHandlerSubsystem::HasObjectHandler(const UClass* InClass) const
diff --git a/Engine/Plugins/Experimental/MediaStream/Source/MediaStream/Private/MediaStreamSchemeHandlerLibrary.cpp b/Engine/Plugins/Experimental/MediaStream/Source/MediaStream/Private/MediaStreamSchemeHandlerLibrary.cpp
--- a/Engine/Plugins/Experimental/MediaStream/Source/MediaStream/Private/MediaStreamSchemeHandlerLibrary.cpp
+++ b/Engine/Plugins/Experimental/MediaStream/Source/MediaStream/Private/MediaStreamSchemeHandlerLibrary.cpp
@@ -1,5 +1,6 @@
 // Copyright Epic Games, Inc. All Rights Reserved.
 
+#include "IMediaStreamObjectHandler.h"
 #if WITH_EDITOR
 
 #include "IMediaStreamSchemeHandler.h"
@@ -8,6 +9,17 @@
 #include "MediaStreamSource.h"
 #include "PropertyHandle.h"
 
+FMediaStreamObjectHandlerCreatePlayerParams FMediaStreamSchemeHandlerCreatePlayerParams::operator<<(UObject* InMediaSource) const
+{
+	FMediaStreamObjectHandlerCreatePlayerParams NewParams;
+	NewParams.MediaStream = MediaStream;
+	NewParams.Source = InMediaSource;
+	NewParams.CurrentPlayer = CurrentPlayer;
+	NewParams.bCanOpenSource = bCanOpenSource;
+
+	return NewParams;
+}
+
 FMediaStreamSource* FMediaStreamSchemeHandlerLibrary::GetStreamSourcePtr(TWeakPtr<IPropertyHandle> InPropertyHandleWeak)
 {
 	TSharedPtr<IPropertyHandle> PropertyHandle = InPropertyHandleWeak.Pin();
diff --git a/Engine/Plugins/Experimental/MediaStream/Source/MediaStream/Private/MediaStreamSchemeHandlerManager.cpp b/Engine/Plugins/Experimental/MediaStream/Source/MediaStream/Private/MediaStreamSchemeHandlerManager.cpp
--- a/Engine/Plugins/Experimental/MediaStream/Source/MediaStream/Private/MediaStreamSchemeHandlerManager.cpp
+++ b/Engine/Plugins/Experimental/MediaStream/Source/MediaStream/Private/MediaStreamSchemeHandlerManager.cpp
@@ -68,18 +68,18 @@ TSharedPtr<IMediaStreamSchemeHandler> FMediaStreamSchemeHandlerManager::GetHandl
 	return nullptr;
 }
 
-UMediaPlayer* FMediaStreamSchemeHandlerManager::CreateOrUpdateMediaPlayer(UMediaStream* InMediaStream, UMediaPlayer* InCurrentPlayer) const
+UMediaPlayer* FMediaStreamSchemeHandlerManager::CreateOrUpdatePlayer(const FMediaStreamSchemeHandlerCreatePlayerParams& InParams) const
 {
-	const FMediaStreamSource& Source = InMediaStream->GetSource();
+	const FMediaStreamSource& Source = InParams.MediaStream->GetSource();
 	const TSharedRef<IMediaStreamSchemeHandler>* HandlerPtr = Handlers.Find(Source.Scheme);
 
 	if (!HandlerPtr)
 	{
-		UE_LOG(LogMediaStream, Error, TEXT("No handler for scheme in FMediaStreamSchemeHandlerManager::CreateOrUpdateMediaPlayer [%s]"), *Source.Scheme.ToString());
+		UE_LOG(LogMediaStream, Error, TEXT("No handler for scheme in FMediaStreamSchemeHandlerManager::CreateOrUpdatePlayer [%s]"), *Source.Scheme.ToString());
 		return nullptr;
 	}
 
-	return (*HandlerPtr)->CreateOrUpdateMediaPlayer(InMediaStream, InCurrentPlayer);
+	return (*HandlerPtr)->CreateOrUpdatePlayer(InParams);
 }
 
 FMediaStreamSource FMediaStreamSchemeHandlerManager::CreateSource(UObject* InOuter, FName InScheme, const FString& InPath) const
diff --git a/Engine/Plugins/Experimental/MediaStream/Source/MediaStream/Private/MediaStreamSchemeHandlerSubsystem.cpp b/Engine/Plugins/Experimental/MediaStream/Source/MediaStream/Private/MediaStreamSchemeHandlerSubsystem.cpp
--- a/Engine/Plugins/Experimental/MediaStream/Source/MediaStream/Private/MediaStreamSchemeHandlerSubsystem.cpp
+++ b/Engine/Plugins/Experimental/MediaStream/Source/MediaStream/Private/MediaStreamSchemeHandlerSubsystem.cpp
@@ -42,9 +42,9 @@ TArray<FName> UMediaStreamSchemeHandlerSubsystem::GetSchemeHandlerNames() const
 	return FMediaStreamSchemeHandlerManager::Get().GetSchemeHandlerNames();
 }
 
-UMediaPlayer* UMediaStreamSchemeHandlerSubsystem::CreateOrUpdateMediaPlayer(UMediaStream* InMediaStream, UMediaPlayer* InCurrentPlayer) const
+UMediaPlayer* UMediaStreamSchemeHandlerSubsystem::CreateOrUpdatePlayer(const FMediaStreamSchemeHandlerCreatePlayerParams& InParams) const
 {
-	return FMediaStreamSchemeHandlerManager::Get().CreateOrUpdateMediaPlayer(InMediaStream, InCurrentPlayer);
+	return FMediaStreamSchemeHandlerManager::Get().CreateOrUpdatePlayer(InParams);
 }
 
 FMediaStreamSource UMediaStreamSchemeHandlerSubsystem::CreateSource(UObject* InOuter, FName InScheme, const FString& InPath) const
diff --git a/Engine/Plugins/Experimental/MediaStream/Source/MediaStream/Private/ObjectHandlers/MediaStreamMediaPlaylistHandler.cpp b/Engine/Plugins/Experimental/MediaStream/Source/MediaStream/Private/ObjectHandlers/MediaStreamMediaPlaylistHandler.cpp
--- a/Engine/Plugins/Experimental/MediaStream/Source/MediaStream/Private/ObjectHandlers/MediaStreamMediaPlaylistHandler.cpp
+++ b/Engine/Plugins/Experimental/MediaStream/Source/MediaStream/Private/ObjectHandlers/MediaStreamMediaPlaylistHandler.cpp
@@ -12,21 +12,21 @@ UClass* FMediaStreamMediaPlaylistHandler::GetClass()
 	return UMediaPlaylist::StaticClass();
 }
 
-UMediaPlayer* FMediaStreamMediaPlaylistHandler::CreateOrUpdatePlayer(UMediaStream* InMediaStream, UObject* InObject, UMediaPlayer* InCurrentPlayer)
+UMediaPlayer* FMediaStreamMediaPlaylistHandler::CreateOrUpdatePlayer(const FMediaStreamObjectHandlerCreatePlayerParams& InParams)
 {
-	if (!InMediaStream)
+	if (!InParams.MediaStream)
 	{
 		UE_LOG(LogMediaStream, Error, TEXT("Invalid Media Stream in FMediaStreamMediaPlaylistHandler::CreateOrUpdatePlayer"));
 		return nullptr;
 	}
 
-	if (!InObject)
+	if (!InParams.Source)
 	{
 		UE_LOG(LogMediaStream, Error, TEXT("Invalid Source Object in FMediaStreamMediaPlaylistHandler::CreateOrUpdatePlayer"));
 		return nullptr;
 	}
 
-	UMediaPlaylist* MediaPlaylist = Cast<UMediaPlaylist>(InObject);
+	UMediaPlaylist* MediaPlaylist = Cast<UMediaPlaylist>(InParams.Source);
 
 	if (!MediaPlaylist)
 	{
@@ -34,12 +34,27 @@ UMediaPlayer* FMediaStreamMediaPlaylistHandler::CreateOrUpdatePlayer(UMediaStrea
 		return nullptr;
 	}
 
-	if (InCurrentPlayer && InCurrentPlayer->OpenPlaylist(MediaPlaylist))
+	if (InParams.CurrentPlayer)
 	{
-		return InCurrentPlayer;
+		UMediaSource* FirstPlaylistEntry = MediaPlaylist->Get(0);
+
+		const bool bIsValidPlayer = InParams.bCanOpenSource
+			? InParams.CurrentPlayer->OpenPlaylist(MediaPlaylist)
+			: (!FirstPlaylistEntry || InParams.CurrentPlayer->CanPlaySource(FirstPlaylistEntry));
+
+		if (bIsValidPlayer)
+		{
+			return InParams.CurrentPlayer;
+		}
+	}
+
+	if (!InParams.bCanOpenSource)
+	{
+		UE_LOG(LogMediaStream, Error, TEXT("Cannot create new player at the moment in FMediaStreamMediaPlaylistHandler::CreateOrUpdatePlayer"));
+		return nullptr;
 	}
 
-	UMediaPlayer* MediaPlayer = NewObject<UMediaPlayer>(InMediaStream);
+	UMediaPlayer* MediaPlayer = NewObject<UMediaPlayer>(InParams.MediaStream);
 
 	if (!MediaPlayer->OpenPlaylist(MediaPlaylist))
 	{
diff --git a/Engine/Plugins/Experimental/MediaStream/Source/MediaStream/Private/ObjectHandlers/MediaStreamMediaPlaylistHandler.h b/Engine/Plugins/Experimental/MediaStream/Source/MediaStream/Private/ObjectHandlers/MediaStreamMediaPlaylistHandler.h
--- a/Engine/Plugins/Experimental/MediaStream/Source/MediaStream/Private/ObjectHandlers/MediaStreamMediaPlaylistHandler.h
+++ b/Engine/Plugins/Experimental/MediaStream/Source/MediaStream/Private/ObjectHandlers/MediaStreamMediaPlaylistHandler.h
@@ -10,6 +10,6 @@ public:
 	static UClass* GetClass();
 
 	//~ Begin IMediaStreamObjectHandler
-	virtual UMediaPlayer* CreateOrUpdatePlayer(UMediaStream* InMediaStream, UObject* InObject, UMediaPlayer* InCurrentPlayer) override;
+	virtual UMediaPlayer* CreateOrUpdatePlayer(const FMediaStreamObjectHandlerCreatePlayerParams& InParams) override;
 	//~ End IMediaStreamObjectHandler
 };
diff --git a/Engine/Plugins/Experimental/MediaStream/Source/MediaStream/Private/ObjectHandlers/MediaStreamMediaSourceHandler.cpp b/Engine/Plugins/Experimental/MediaStream/Source/MediaStream/Private/ObjectHandlers/MediaStreamMediaSourceHandler.cpp
--- a/Engine/Plugins/Experimental/MediaStream/Source/MediaStream/Private/ObjectHandlers/MediaStreamMediaSourceHandler.cpp
+++ b/Engine/Plugins/Experimental/MediaStream/Source/MediaStream/Private/ObjectHandlers/MediaStreamMediaSourceHandler.cpp
@@ -13,21 +13,21 @@ UClass* FMediaStreamMediaSourceHandler::GetClass()
 	return UMediaSource::StaticClass();
 }
 
-UMediaPlayer* FMediaStreamMediaSourceHandler::CreateOrUpdatePlayer(UMediaStream* InMediaStream, UObject* InObject, UMediaPlayer* InCurrentPlayer)
+UMediaPlayer* FMediaStreamMediaSourceHandler::CreateOrUpdatePlayer(const FMediaStreamObjectHandlerCreatePlayerParams& InParams)
 {
-	if (!InMediaStream)
+	if (!InParams.MediaStream)
 	{
 		UE_LOG(LogMediaStream, Error, TEXT("Invalid Media Stream in FMediaStreamMediaSourceHandler::CreateOrUpdatePlayer"));
 		return nullptr;
 	}
 
-	if (!InObject)
+	if (!InParams.Source)
 	{
 		UE_LOG(LogMediaStream, Error, TEXT("Invalid Source Object in FMediaStreamMediaSourceHandler::CreateOrUpdatePlayer"));
 		return nullptr;
 	}
 
-	IMediaStreamPlayer* MediaStreamPlayer = InMediaStream->GetPlayer().GetInterface();
+	IMediaStreamPlayer* MediaStreamPlayer = InParams.MediaStream->GetPlayer().GetInterface();
 
 	if (!MediaStreamPlayer)
 	{
@@ -35,7 +35,7 @@ UMediaPlayer* FMediaStreamMediaSourceHandler::CreateOrUpdatePlayer(UMediaStream*
 		return nullptr;
 	}
 
-	UMediaSource* MediaSource = Cast<UMediaSource>(InObject);
+	UMediaSource* MediaSource = Cast<UMediaSource>(InParams.Source);
 
 	if (!MediaSource)
 	{
@@ -43,12 +43,25 @@ UMediaPlayer* FMediaStreamMediaSourceHandler::CreateOrUpdatePlayer(UMediaStream*
 		return nullptr;
 	}
 
-	if (InCurrentPlayer && InCurrentPlayer->OpenSource(MediaSource))
+	if (InParams.CurrentPlayer)
 	{
-		return InCurrentPlayer;
+		const bool bIsValidPlayer = InParams.bCanOpenSource
+			? InParams.CurrentPlayer->OpenSource(MediaSource)
+			: InParams.CurrentPlayer->CanPlaySource(MediaSource);
+
+		if (bIsValidPlayer)
+		{
+			return InParams.CurrentPlayer;
+		}
+	}
+
+	if (!InParams.bCanOpenSource)
+	{
+		UE_LOG(LogMediaStream, Error, TEXT("Cannot create new player at the moment in FMediaStreamMediaSourceHandler::CreateOrUpdatePlayer"));
+		return nullptr;
 	}
 
-	UMediaPlayer* MediaPlayer = NewObject<UMediaPlayer>(InMediaStream, NAME_None, RF_Transactional);
+	UMediaPlayer* MediaPlayer = NewObject<UMediaPlayer>(InParams.MediaStream, NAME_None, RF_Transactional);
 	const FMediaPlayerOptions Options = MediaStreamPlayer->GetPlayerConfig().CreateOptions(MediaStreamPlayer->GetRequestedSeekTime());
 
 	if (!MediaPlayer->OpenSourceWithOptions(MediaSource, Options))
diff --git a/Engine/Plugins/Experimental/MediaStream/Source/MediaStream/Private/ObjectHandlers/MediaStreamMediaSourceHandler.h b/Engine/Plugins/Experimental/MediaStream/Source/MediaStream/Private/ObjectHandlers/MediaStreamMediaSourceHandler.h
--- a/Engine/Plugins/Experimental/MediaStream/Source/MediaStream/Private/ObjectHandlers/MediaStreamMediaSourceHandler.h
+++ b/Engine/Plugins/Experimental/MediaStream/Source/MediaStream/Private/ObjectHandlers/MediaStreamMediaSourceHandler.h
@@ -10,6 +10,6 @@ public:
 	static UClass* GetClass();
 
 	//~ Begin IMediaStreamObjectHandler
-	virtual UMediaPlayer* CreateOrUpdatePlayer(UMediaStream* InMediaStream, UObject* InObject, UMediaPlayer* InCurrentPlayer) override;
+	virtual UMediaPlayer* CreateOrUpdatePlayer(const FMediaStreamObjectHandlerCreatePlayerParams& InParams) override;
 	//~ End IMediaStreamObjectHandler
 };
diff --git a/Engine/Plugins/Experimental/MediaStream/Source/MediaStream/Private/ObjectHandlers/MediaStreamMediaStreamHandler.cpp b/Engine/Plugins/Experimental/MediaStream/Source/MediaStream/Private/ObjectHandlers/MediaStreamMediaStreamHandler.cpp
--- a/Engine/Plugins/Experimental/MediaStream/Source/MediaStream/Private/ObjectHandlers/MediaStreamMediaStreamHandler.cpp
+++ b/Engine/Plugins/Experimental/MediaStream/Source/MediaStream/Private/ObjectHandlers/MediaStreamMediaStreamHandler.cpp
@@ -12,21 +12,21 @@ UClass* FMediaStreamMediaStreamHandler::GetClass()
 	return UMediaStream::StaticClass();
 }
 
-UMediaPlayer* FMediaStreamMediaStreamHandler::CreateOrUpdatePlayer(UMediaStream* InMediaStream, UObject* InObject, UMediaPlayer* InCurrentPlayer)
+UMediaPlayer* FMediaStreamMediaStreamHandler::CreateOrUpdatePlayer(const FMediaStreamObjectHandlerCreatePlayerParams& InParams)
 {
-	if (!InMediaStream)
+	if (!InParams.MediaStream)
 	{
 		UE_LOG(LogMediaStream, Error, TEXT("Invalid Media Stream in FMediaStreamMediaStreamHandler::CreateOrUpdatePlayer"));
 		return nullptr;
 	}
 
-	if (!InObject)
+	if (!InParams.Source)
 	{
 		UE_LOG(LogMediaStream, Error, TEXT("Invalid Source Object in FMediaStreamMediaStreamHandler::CreateOrUpdatePlayer"));
 		return nullptr;
 	}
 
-	UMediaStream* MediaStream = Cast<UMediaStream>(InObject);
+	UMediaStream* MediaStream = Cast<UMediaStream>(InParams.Source);
 
 	if (!MediaStream)
 	{
diff --git a/Engine/Plugins/Experimental/MediaStream/Source/MediaStream/Private/ObjectHandlers/MediaStreamMediaStreamHandler.h b/Engine/Plugins/Experimental/MediaStream/Source/MediaStream/Private/ObjectHandlers/MediaStreamMediaStreamHandler.h
--- a/Engine/Plugins/Experimental/MediaStream/Source/MediaStream/Private/ObjectHandlers/MediaStreamMediaStreamHandler.h
+++ b/Engine/Plugins/Experimental/MediaStream/Source/MediaStream/Private/ObjectHandlers/MediaStreamMediaStreamHandler.h
@@ -10,6 +10,6 @@ public:
 	static UClass* GetClass();
 
 	//~ Begin IMediaStreamObjectHandler
-	virtual UMediaPlayer* CreateOrUpdatePlayer(UMediaStream* InMediaStream, UObject* InObject, UMediaPlayer* InCurrentPlayer) override;
+	virtual UMediaPlayer* CreateOrUpdatePlayer(const FMediaStreamObjectHandlerCreatePlayerParams& InParams) override;
 	//~ End IMediaStreamObjectHandler
 };
diff --git a/Engine/Plugins/Experimental/MediaStream/Source/MediaStream/Private/Players/MediaStreamLocalPlayer.cpp b/Engine/Plugins/Experimental/MediaStream/Source/MediaStream/Private/Players/MediaStreamLocalPlayer.cpp
--- a/Engine/Plugins/Experimental/MediaStream/Source/MediaStream/Private/Players/MediaStreamLocalPlayer.cpp
+++ b/Engine/Plugins/Experimental/MediaStream/Source/MediaStream/Private/Players/MediaStreamLocalPlayer.cpp
@@ -3,6 +3,7 @@
 #include "Players/MediaStreamLocalPlayer.h"
 
 #include "CoreGlobals.h"
+#include "IMediaStreamSchemeHandler.h"
 #include "LevelSequence.h"
 #include "MediaPlaylist.h"
 #include "MediaStream.h"
@@ -49,10 +50,6 @@ void UMediaStreamLocalPlayer::Deinitialize()
 {
 	DeinitPlayer();
 	DeinitTexture();
-
-#if WITH_EDITOR
-	DeinitEditorPlaybackStateEvent();
-#endif
 }
 
 #if WITH_EDITOR
@@ -136,6 +133,8 @@ void UMediaStreamLocalPlayer::PostLoad()
 {
 	Super::PostLoad();
 
+	bAllowOpenSource = FMediaStreamModule::Get().CanOpenSourceOnLoad();
+
 	Initialize();
 }
 
@@ -155,9 +154,7 @@ void UMediaStreamLocalPlayer::BeginDestroy()
 
 void UMediaStreamLocalPlayer::OnSourceChanged(const FMediaStreamSource& InSource)
 {
-#if WITH_EDITOR
-	bPlayerSourceNeedsUpdate = true;
-#endif
+	bPlayerNeedsUpdate = true;
 
 	Initialize();
 }
@@ -230,23 +227,30 @@ void UMediaStreamLocalPlayer::ApplyPlayerConfig()
 
 	// Stop media playing on open in the editor (only in PIE)
 #if WITH_EDITOR
-	bool bOverridePlayOnOpen = GIsEditor && !bIsPIE; // Handle -game
-	bool bInitialPlayOnOpen = PlayerConfig.bPlayOnOpen;
+	const bool bOverridePlayOnOpen = !bIsPIE && !FMediaStreamModule::Get().CanAutoplay();
+#else
+	const bool bOverridePlayOnOpen = !FMediaStreamModule::Get().CanAutoplay();
+#endif
+
+	if (!bOverridePlayOnOpen)
+	{
+		PlayerConfig.ApplyConfig(*Player);
+		return;
+	}
+
+	const bool bInitialPlayOnOpen = PlayerConfig.bPlayOnOpen;
 
 	if (bOverridePlayOnOpen)
 	{
 		PlayerConfig.bPlayOnOpen = false;
 	}
-#endif
 
 	PlayerConfig.ApplyConfig(*Player);
 
-#if WITH_EDITOR
 	if (bOverridePlayOnOpen)
 	{
 		PlayerConfig.bPlayOnOpen = bInitialPlayOnOpen;
 	}
-#endif
 }
 
 float UMediaStreamLocalPlayer::GetRequestedSeekTime() const
@@ -295,12 +299,17 @@ bool UMediaStreamLocalPlayer::ApplyPlaybackState()
 			case EMediaStreamPlaybackState::Play:
 				if (UMediaTimeStampInfo* TimeStamp = Player->GetTimeStamp())
 				{
-					const double TotalTime = Player->GetDuration().GetTotalSeconds();
-					const double CurrentTime = TimeStamp->Time.GetTotalSeconds();
+					const float FrameRate = Player->GetVideoTrackFrameRate(INDEX_NONE, INDEX_NONE);
 
-					if (!PlayerConfig.bLooping && FMath::IsNearlyEqual(TotalTime, CurrentTime, Player->GetVideoTrackFrameRate(INDEX_NONE, INDEX_NONE) * 0.5))
+					if (!FMath::IsNearlyZero(FrameRate))
 					{
-						Player->Seek(0);
+						const double TotalTime = Player->GetDuration().GetTotalSeconds();
+						const double CurrentTime = TimeStamp->Time.GetTotalSeconds();
+
+						if (!PlayerConfig.bLooping && FMath::Abs(TotalTime - CurrentTime) < (2.0 / FrameRate))
+						{
+							Player->Seek(0);
+						}
 					}
 				}
 
@@ -353,6 +362,15 @@ UMediaStream* UMediaStreamLocalPlayer::GetSourceStream() const
 	return GetMediaStream();
 }
 
+bool UMediaStreamLocalPlayer::OpenSource()
+{
+	bAllowOpenSource = true;
+	bPlayerNeedsUpdate = true;
+	InitPlayer();
+
+	return Player && !Player->IsClosed();
+}
+
 bool UMediaStreamLocalPlayer::Play()
 {
 	return SetPlaybackState(EMediaStreamPlaybackState::Play);
@@ -446,10 +464,6 @@ void UMediaStreamLocalPlayer::Initialize()
 {
 	InitTexture();
 	InitPlayer();
-
-#if WITH_EDITOR
-	InitEditorPlaybackStateEvent();
-#endif
 }
 
 void UMediaStreamLocalPlayer::InitTexture()
@@ -466,29 +480,22 @@ void UMediaStreamLocalPlayer::DeinitTexture()
 
 void UMediaStreamLocalPlayer::InitPlayer()
 {
-#if WITH_EDITOR
 	UMediaPlayer* CurrentPlayer = Player.Get();
 
-	if (bPlayerSourceNeedsUpdate || !IsValid(Player))
+	if (bPlayerNeedsUpdate || !IsValid(Player))
 	{
-		Player = FMediaStreamSchemeHandlerManager::Get().CreateOrUpdateMediaPlayer(GetMediaStream(), Player.Get());
-		bPlayerSourceNeedsUpdate = false;
+#if WITH_EDITOR
+		Player = FMediaStreamSchemeHandlerManager::Get().CreateOrUpdatePlayer({GetMediaStream(), Player.Get(), bIsPIE || bAllowOpenSource});
+#else
+		Player = FMediaStreamSchemeHandlerManager::Get().CreateOrUpdatePlayer({GetMediaStream(), Player.Get(), bAllowOpenSource});
+#endif
+		bPlayerNeedsUpdate = false;
 	}	
 
 	if (CurrentPlayer != Player)
 	{
 		UpdateSequencesWithCurrentPlayer();
 	}
-#else
-	if (!IsValid(Player))
-	{
-		return;
-	}
-	else
-	{
-		UpdateSequencesWithCurrentPlayer();
-	}
-#endif
 
 	const bool bIsPlayerValid = IsValid(Player);
 
@@ -563,21 +570,21 @@ void UMediaStreamLocalPlayer::UpdateSequencesWithCurrentPlayer()
 			continue;
 		}
 
-		ULevelSequence* LevelSequence = Cast<ULevelSequence>(MovieScene->GetOuter());
+		UMovieSceneSequence* MovieSequence = Cast<UMovieSceneSequence>(MovieScene->GetOuter());
 
-		if (!LevelSequence)
+		if (!MovieSequence)
 		{
 			continue;
 		}
 
-		UWorld* World = LevelSequence->GetWorld();
+		UWorld* World = MovieSequence->GetWorld();
 
 		if (!World)
 		{
 			continue;
 		}
 
-		const FMovieSceneBindingReferences* References = LevelSequence->GetBindingReferences();
+		const FMovieSceneBindingReferences* References = MovieSequence->GetBindingReferences();
 
 		if (!References)
 		{
@@ -594,56 +601,8 @@ void UMediaStreamLocalPlayer::UpdateSequencesWithCurrentPlayer()
 			continue;
 		}
 
-		Section->ExternalMediaPlayer = Player;
-	}
-}
-
-#if WITH_EDITOR
-void UMediaStreamLocalPlayer::InitEditorPlaybackStateEvent()
-{
-	if (!GIsEditor)
-	{
-		return;
-	}
-
-	FMediaStreamEditorPlaybackStateDelegate::RegistrationType& PIEEvent = FMediaStreamModule::Get().GetEditorPlaybackStateEvent();
-
-	if (!PIEEvent.IsBoundToObject(this))
-	{
-		PIEEvent.AddUObject(this, &UMediaStreamLocalPlayer::OnEditorPlaybackStateEvent);
-	}
-}
-
-void UMediaStreamLocalPlayer::DeinitEditorPlaybackStateEvent()
-{
-	if (!GIsEditor)
-	{
-		return;
-	}
-
-	FMediaStreamModule::Get().GetEditorPlaybackStateEvent().RemoveAll(this);
-}
-
-void UMediaStreamLocalPlayer::OnEditorPlaybackStateEvent(EMediaStreamPlaybackState InPlaybackState)
-{
-	if (bIsPIE || PlaybackState == EMediaStreamPlaybackState::Pause || !IsValid(Player))
-	{
-		return;
-	}
-
-	if (InPlaybackState == EMediaStreamPlaybackState::Play)
-	{
-		if (!FMath::IsNearlyZero(PlayerConfig.Rate))
-		{
-			Player->SetRate(PlayerConfig.Rate);
-		}
-	}
-	else if (InPlaybackState == EMediaStreamPlaybackState::Pause)
-	{
-		if (!FMath::IsNearlyZero(Player->GetRate()))
-		{
-			Player->Pause();
-		}
+		Section->bUseExternalMediaPlayer = true;
+		Section->ExternalMediaPlayer = Player;		
+		Section->TryModify();
 	}
 }
-#endif
diff --git a/Engine/Plugins/Experimental/MediaStream/Source/MediaStream/Private/Players/MediaStreamLocalPlayer.h b/Engine/Plugins/Experimental/MediaStream/Source/MediaStream/Private/Players/MediaStreamLocalPlayer.h
--- a/Engine/Plugins/Experimental/MediaStream/Source/MediaStream/Private/Players/MediaStreamLocalPlayer.h
+++ b/Engine/Plugins/Experimental/MediaStream/Source/MediaStream/Private/Players/MediaStreamLocalPlayer.h
@@ -40,6 +40,8 @@ public:
 		SetPlaylistIndex(InIndex);
 	}
 
+	void LoadSource();
+
 	//~ Begin IMediaStreamPlayer
 	virtual UMediaStream* GetMediaStream() const override;
 	virtual bool IsReadOnly() const override;
@@ -62,6 +64,7 @@ public:
 	virtual int32 GetPlaylistIndex() const override;
 	virtual int32 GetPlaylistNum() const override;
 	virtual UMediaStream* GetSourceStream() const override;
+	virtual bool OpenSource() override;
 	virtual bool Play() override;
 	virtual bool Pause() override;
 	virtual bool Rewind() override;
@@ -117,9 +120,18 @@ protected:
 	FMediaStreamTextureConfig TextureConfig_PreUndo;
 	FMediaStreamPlayerConfig PlayerConfig_PreUndo;
 
-#if WITH_EDITOR
+	bool bPlayerNeedsUpdate = false;
+
+	/**
+	 * Determines whether local players can create or update media players when opening a source.
+	 */
+	bool bAllowOpenSource = true;
+
+	#if WITH_EDITOR
+	/**
+	 * Whether this object was duplicated into a PIE level.
+	 */
 	bool bIsPIE = false;
-	bool bPlayerSourceNeedsUpdate = false;
 #endif
 
 	/**
@@ -173,22 +185,4 @@ protected:
 	void OnMediaOpened(FString InOpenedUrl);
 
 	void UpdateSequencesWithCurrentPlayer();
-
-#if WITH_EDITOR
-	/**
-	 * Initializes PIE events
-	 */
-	void InitEditorPlaybackStateEvent();
-
-	/**
-	 * Deinitializes PIE events
-	 */
-	void DeinitEditorPlaybackStateEvent();
-
-	/**
-	 * Called when the editor button to start and stop playback in the editor is pressed.
-	 * @param InPlaybackState
-	 */
-	void OnEditorPlaybackStateEvent(EMediaStreamPlaybackState InPlaybackState);
-#endif
 };
diff --git a/Engine/Plugins/Experimental/MediaStream/Source/MediaStream/Private/Players/MediaStreamProxyPlayer.cpp b/Engine/Plugins/Experimental/MediaStream/Source/MediaStream/Private/Players/MediaStreamProxyPlayer.cpp
--- a/Engine/Plugins/Experimental/MediaStream/Source/MediaStream/Private/Players/MediaStreamProxyPlayer.cpp
+++ b/Engine/Plugins/Experimental/MediaStream/Source/MediaStream/Private/Players/MediaStreamProxyPlayer.cpp
@@ -304,6 +304,22 @@ int32 UMediaStreamProxyPlayer::GetPlaylistNum() const
 	return -1;
 }
 
+bool UMediaStreamProxyPlayer::OpenSource()
+{
+	if (!bReadOnly)
+	{
+		if (UMediaStream* MediaStream = GetSourceStream())
+		{
+			if (TScriptInterface<IMediaStreamPlayer> MediaStreamPlayer = MediaStream->GetPlayer())
+			{
+				return MediaStreamPlayer->OpenSource();
+			}
+		}
+	}
+
+	return false;
+}
+
 bool UMediaStreamProxyPlayer::Play()
 {
 	if (!bReadOnly)
diff --git a/Engine/Plugins/Experimental/MediaStream/Source/MediaStream/Private/Players/MediaStreamProxyPlayer.h b/Engine/Plugins/Experimental/MediaStream/Source/MediaStream/Private/Players/MediaStreamProxyPlayer.h
--- a/Engine/Plugins/Experimental/MediaStream/Source/MediaStream/Private/Players/MediaStreamProxyPlayer.h
+++ b/Engine/Plugins/Experimental/MediaStream/Source/MediaStream/Private/Players/MediaStreamProxyPlayer.h
@@ -46,6 +46,7 @@ public:
 	virtual int32 GetPlaylistIndex() const override;
 	virtual int32 GetPlaylistNum() const override;
 	virtual UMediaStream* GetSourceStream() const override;
+	virtual bool OpenSource() override;
 	virtual bool Play() override;
 	virtual bool Pause() override;
 	virtual bool Rewind() override;
diff --git a/Engine/Plugins/Experimental/MediaStream/Source/MediaStream/Private/SchemeHandlers/MediaStreamAssetSchemeHandler.cpp b/Engine/Plugins/Experimental/MediaStream/Source/MediaStream/Private/SchemeHandlers/MediaStreamAssetSchemeHandler.cpp
--- a/Engine/Plugins/Experimental/MediaStream/Source/MediaStream/Private/SchemeHandlers/MediaStreamAssetSchemeHandler.cpp
+++ b/Engine/Plugins/Experimental/MediaStream/Source/MediaStream/Private/SchemeHandlers/MediaStreamAssetSchemeHandler.cpp
@@ -2,6 +2,7 @@
 
 #include "SchemeHandlers/MediaStreamAssetSchemeHandler.h"
 
+#include "IMediaStreamObjectHandler.h"
 #include "MediaSource.h"
 #include "MediaStream.h"
 #include "MediaStreamModule.h"
@@ -53,23 +54,23 @@ FMediaStreamSource FMediaStreamAssetSchemeHandler::CreateSource(UObject* InOuter
 	return Source;
 }
 
-UMediaPlayer* FMediaStreamAssetSchemeHandler::CreateOrUpdateMediaPlayer(UMediaStream* InMediaStream, UMediaPlayer* InCurrentPlayer)
+UMediaPlayer* FMediaStreamAssetSchemeHandler::CreateOrUpdatePlayer(const FMediaStreamSchemeHandlerCreatePlayerParams& InParams)
 {
-	if (!IsValid(InMediaStream))
+	if (!IsValid(InParams.MediaStream))
 	{
-		UE_LOG(LogMediaStream, Error, TEXT("Invalid Media Stream in FMediaStreamAssetSchemeHandler::CreateOrUpdateMediaPlayer"));
+		UE_LOG(LogMediaStream, Error, TEXT("Invalid Media Stream in FMediaStreamAssetSchemeHandler::CreateOrUpdatePlayer"));
 		return nullptr;
 	}
 
-	const FMediaStreamSource& Source = InMediaStream->GetSource();
+	const FMediaStreamSource& Source = InParams.MediaStream->GetSource();
 
 	if (!IsValid(Source.Object))
 	{
-		UE_LOG(LogMediaStream, Error, TEXT("Invalid Asset in FMediaStreamAssetSchemeHandler::CreateOrUpdateMediaPlayer"));
+		UE_LOG(LogMediaStream, Error, TEXT("Invalid Asset in FMediaStreamAssetSchemeHandler::CreateOrUpdatePlayer"));
 		return nullptr;
 	}
 
-	return FMediaStreamObjectHandlerManager::Get().CreateOrUpdateMediaPlayer(InMediaStream, Source.Object, InCurrentPlayer);
+	return FMediaStreamObjectHandlerManager::Get().CreateOrUpdatePlayer(InParams << Source.Object);
 }
 
 UObject* FMediaStreamAssetSchemeHandler::ResolveAsset(const FString& InPath)
@@ -143,6 +144,7 @@ void FMediaStreamAssetSchemeHandler::AddAssetSelector(TSharedRef<IPropertyHandle
 		.ThumbnailPool(InCustomizationUtils.GetThumbnailPool())
 		.DisplayThumbnail(true)
 		.AllowedClass(UMediaSource::StaticClass())
+		.PropertyHandle(InPropertyHandle)
 	];
 }
 
diff --git a/Engine/Plugins/Experimental/MediaStream/Source/MediaStream/Private/SchemeHandlers/MediaStreamAssetSchemeHandler.h b/Engine/Plugins/Experimental/MediaStream/Source/MediaStream/Private/SchemeHandlers/MediaStreamAssetSchemeHandler.h
--- a/Engine/Plugins/Experimental/MediaStream/Source/MediaStream/Private/SchemeHandlers/MediaStreamAssetSchemeHandler.h
+++ b/Engine/Plugins/Experimental/MediaStream/Source/MediaStream/Private/SchemeHandlers/MediaStreamAssetSchemeHandler.h
@@ -18,7 +18,7 @@ public:
 
 	//~ Begin IMediaStreamSchemeHandler
 	virtual FMediaStreamSource CreateSource(UObject* InOuter, const FString& InPath) override;
-	virtual UMediaPlayer* CreateOrUpdateMediaPlayer(UMediaStream* InMediaStream, UMediaPlayer* InCurrentPlayer) override;
+	virtual UMediaPlayer* CreateOrUpdatePlayer(const FMediaStreamSchemeHandlerCreatePlayerParams& InParams) override;
 
 #if WITH_EDITOR
 	virtual void CreatePropertyCustomization(TSharedRef<IPropertyHandle> InPropertyHandle, IDetailGroup& InDetailsGroup, IPropertyTypeCustomizationUtils& InCustomizationUtils) override;
diff --git a/Engine/Plugins/Experimental/MediaStream/Source/MediaStream/Private/SchemeHandlers/MediaStreamFileSchemeHandler.cpp b/Engine/Plugins/Experimental/MediaStream/Source/MediaStream/Private/SchemeHandlers/MediaStreamFileSchemeHandler.cpp
--- a/Engine/Plugins/Experimental/MediaStream/Source/MediaStream/Private/SchemeHandlers/MediaStreamFileSchemeHandler.cpp
+++ b/Engine/Plugins/Experimental/MediaStream/Source/MediaStream/Private/SchemeHandlers/MediaStreamFileSchemeHandler.cpp
@@ -3,6 +3,7 @@
 #include "SchemeHandlers/MediaStreamFileSchemeHandler.h"
 
 #include "HAL/FileManager.h"
+#include "IMediaStreamObjectHandler.h"
 #include "MediaSource.h"
 #include "MediaStream.h"
 #include "MediaStreamModule.h"
@@ -47,19 +48,19 @@ FMediaStreamSource FMediaStreamFileSchemeHandler::CreateSource(UObject* InOuter,
 	return Source;
 }
 
-UMediaPlayer* FMediaStreamFileSchemeHandler::CreateOrUpdateMediaPlayer(UMediaStream* InMediaStream, UMediaPlayer* InCurrentPlayer)
+UMediaPlayer* FMediaStreamFileSchemeHandler::CreateOrUpdatePlayer(const FMediaStreamSchemeHandlerCreatePlayerParams& InParams)
 {
-	if (!IsValid(InMediaStream))
+	if (!IsValid(InParams.MediaStream))
 	{
-		UE_LOG(LogMediaStream, Error, TEXT("Invalid Media Stream in FMediaStreamFileSchemeHandler::CreateOrUpdateMediaPlayer"));
+		UE_LOG(LogMediaStream, Error, TEXT("Invalid Media Stream in FMediaStreamFileSchemeHandler::CreateOrUpdatePlayer"));
 		return nullptr;
 	}
 
-	UObject* MediaObject = InMediaStream->GetSource().Object;
+	UObject* MediaObject = InParams.MediaStream->GetSource().Object;
 
 	if (!IsValid(MediaObject))
 	{
-		MediaObject = CreateMediaSource(InMediaStream, InMediaStream->GetSource().Path);
+		MediaObject = CreateMediaSource(InParams.MediaStream, InParams.MediaStream->GetSource().Path);
 
 		if (!IsValid(MediaObject))
 		{
@@ -71,11 +72,11 @@ UMediaPlayer* FMediaStreamFileSchemeHandler::CreateOrUpdateMediaPlayer(UMediaStr
 
 	if (!IsValid(MediaSource))
 	{
-		UE_LOG(LogMediaStream, Error, TEXT("Invalid Media Path in FMediaStreamFileSchemeHandler::CreateOrUpdateMediaPlayer [%s]"), *MediaObject->GetClass()->GetName());
+		UE_LOG(LogMediaStream, Error, TEXT("Invalid Media Path in FMediaStreamFileSchemeHandler::CreateOrUpdatePlayer [%s]"), *MediaObject->GetClass()->GetName());
 		return nullptr;
 	}
 
-	return FMediaStreamObjectHandlerManager::Get().CreateOrUpdateMediaPlayer(InMediaStream, MediaSource, InCurrentPlayer);
+	return FMediaStreamObjectHandlerManager::Get().CreateOrUpdatePlayer(InParams << MediaSource);
 }
 
 FString FMediaStreamFileSchemeHandler::ResolveFilePath(const FString& InPath)
diff --git a/Engine/Plugins/Experimental/MediaStream/Source/MediaStream/Private/SchemeHandlers/MediaStreamFileSchemeHandler.h b/Engine/Plugins/Experimental/MediaStream/Source/MediaStream/Private/SchemeHandlers/MediaStreamFileSchemeHandler.h
--- a/Engine/Plugins/Experimental/MediaStream/Source/MediaStream/Private/SchemeHandlers/MediaStreamFileSchemeHandler.h
+++ b/Engine/Plugins/Experimental/MediaStream/Source/MediaStream/Private/SchemeHandlers/MediaStreamFileSchemeHandler.h
@@ -14,7 +14,7 @@ public:
 
 	//~ Begin IMediaStreamSchemeHandler
 	virtual FMediaStreamSource CreateSource(UObject* InOuter, const FString& InPath) override;
-	virtual UMediaPlayer* CreateOrUpdateMediaPlayer(UMediaStream* InMediaStream, UMediaPlayer* InCurrentPlayer) override;
+	virtual UMediaPlayer* CreateOrUpdatePlayer(const FMediaStreamSchemeHandlerCreatePlayerParams& InParams) override;
 
 #if WITH_EDITOR
 	virtual void CreatePropertyCustomization(TSharedRef<IPropertyHandle> InPropertyHandle, IDetailGroup& InDetailsGroup, IPropertyTypeCustomizationUtils& InCustomizationUtils) override;
diff --git a/Engine/Plugins/Experimental/MediaStream/Source/MediaStream/Private/SchemeHandlers/MediaStreamManagedSchemeHandler.cpp b/Engine/Plugins/Experimental/MediaStream/Source/MediaStream/Private/SchemeHandlers/MediaStreamManagedSchemeHandler.cpp
--- a/Engine/Plugins/Experimental/MediaStream/Source/MediaStream/Private/SchemeHandlers/MediaStreamManagedSchemeHandler.cpp
+++ b/Engine/Plugins/Experimental/MediaStream/Source/MediaStream/Private/SchemeHandlers/MediaStreamManagedSchemeHandler.cpp
@@ -32,7 +32,7 @@ FMediaStreamSource FMediaStreamManagedSchemeHandler::CreateSource(UObject* InOut
 	return Source;
 }
 
-UMediaPlayer* FMediaStreamManagedSchemeHandler::CreateOrUpdateMediaPlayer(UMediaStream* InMediaStream, UMediaPlayer* InCurrentPlayer)
+UMediaPlayer* FMediaStreamManagedSchemeHandler::CreateOrUpdatePlayer(const FMediaStreamSchemeHandlerCreatePlayerParams& InParams)
 {
 	// TODO: Get player from media source manager
 	return nullptr;
diff --git a/Engine/Plugins/Experimental/MediaStream/Source/MediaStream/Private/SchemeHandlers/MediaStreamManagedSchemeHandler.h b/Engine/Plugins/Experimental/MediaStream/Source/MediaStream/Private/SchemeHandlers/MediaStreamManagedSchemeHandler.h
--- a/Engine/Plugins/Experimental/MediaStream/Source/MediaStream/Private/SchemeHandlers/MediaStreamManagedSchemeHandler.h
+++ b/Engine/Plugins/Experimental/MediaStream/Source/MediaStream/Private/SchemeHandlers/MediaStreamManagedSchemeHandler.h
@@ -11,7 +11,7 @@ public:
 
 	//~ Begin IMediaStreamSchemeHandler
 	virtual FMediaStreamSource CreateSource(UObject* InOuter, const FString& InPath) override;
-	virtual UMediaPlayer* CreateOrUpdateMediaPlayer(UMediaStream* InMediaStream, UMediaPlayer* InCurrentPlayer) override;
+	virtual UMediaPlayer* CreateOrUpdatePlayer(const FMediaStreamSchemeHandlerCreatePlayerParams& InParams) override;
 
 #if WITH_EDITOR
 	virtual void CreatePropertyCustomization(TSharedRef<IPropertyHandle> InPropertyHandle, IDetailGroup& InDetailsGroup, IPropertyTypeCustomizationUtils& InCustomizationUtils) override;
diff --git a/Engine/Plugins/Experimental/MediaStream/Source/MediaStream/Private/SchemeHandlers/MediaStreamSubobjectSchemeHandler.cpp b/Engine/Plugins/Experimental/MediaStream/Source/MediaStream/Private/SchemeHandlers/MediaStreamSubobjectSchemeHandler.cpp
--- a/Engine/Plugins/Experimental/MediaStream/Source/MediaStream/Private/SchemeHandlers/MediaStreamSubobjectSchemeHandler.cpp
+++ b/Engine/Plugins/Experimental/MediaStream/Source/MediaStream/Private/SchemeHandlers/MediaStreamSubobjectSchemeHandler.cpp
@@ -2,6 +2,7 @@
 
 #include "SchemeHandlers/MediaStreamSubobjectSchemeHandler.h"
 
+#include "IMediaStreamObjectHandler.h"
 #include "MediaStream.h"
 #include "MediaStreamModule.h"
 #include "MediaStreamObjectHandlerManager.h"
@@ -34,23 +35,23 @@ FMediaStreamSource FMediaStreamSubobjectSchemeHandler::CreateSource(UObject* InO
 	return Source;
 }
 
-UMediaPlayer* FMediaStreamSubobjectSchemeHandler::CreateOrUpdateMediaPlayer(UMediaStream* InMediaStream, UMediaPlayer* InCurrentPlayer)
+UMediaPlayer* FMediaStreamSubobjectSchemeHandler::CreateOrUpdatePlayer(const FMediaStreamSchemeHandlerCreatePlayerParams& InParams)
 {
-	if (!IsValid(InMediaStream))
+	if (!IsValid(InParams.MediaStream))
 	{
-		UE_LOG(LogMediaStream, Error, TEXT("Invalid Media Stream in FMediaStreamSubobjectSchemeHandler::CreateOrUpdateMediaPlayer"));
+		UE_LOG(LogMediaStream, Error, TEXT("Invalid Media Stream in FMediaStreamSubobjectSchemeHandler::CreateOrUpdatePlayer"));
 		return nullptr;
 	}
 
-	UObject* Subobject = ResolveSubobjectPath(InMediaStream, InMediaStream->GetSource().Path);
+	UObject* Subobject = ResolveSubobjectPath(InParams.MediaStream, InParams.MediaStream->GetSource().Path);
 
 	if (!IsValid(Subobject))
 	{
-		UE_LOG(LogMediaStream, Error, TEXT("Invalid Subobject in FMediaStreamSubobjectSchemeHandler::CreateOrUpdateMediaPlayer"));
+		UE_LOG(LogMediaStream, Error, TEXT("Invalid Subobject in FMediaStreamSubobjectSchemeHandler::CreateOrUpdatePlayer"));
 		return nullptr;
 	}
 
-	return FMediaStreamObjectHandlerManager::Get().CreateOrUpdateMediaPlayer(InMediaStream, Subobject, InCurrentPlayer);
+	return FMediaStreamObjectHandlerManager::Get().CreateOrUpdatePlayer(InParams << Subobject);
 }
 
 UObject* FMediaStreamSubobjectSchemeHandler::ResolveSubobjectPath(UObject* InOuter, const FString& InPath)
diff --git a/Engine/Plugins/Experimental/MediaStream/Source/MediaStream/Private/SchemeHandlers/MediaStreamSubobjectSchemeHandler.h b/Engine/Plugins/Experimental/MediaStream/Source/MediaStream/Private/SchemeHandlers/MediaStreamSubobjectSchemeHandler.h
--- a/Engine/Plugins/Experimental/MediaStream/Source/MediaStream/Private/SchemeHandlers/MediaStreamSubobjectSchemeHandler.h
+++ b/Engine/Plugins/Experimental/MediaStream/Source/MediaStream/Private/SchemeHandlers/MediaStreamSubobjectSchemeHandler.h
@@ -13,7 +13,7 @@ public:
 
 	//~ Begin IMediaStreamSchemeHandler
 	virtual FMediaStreamSource CreateSource(UObject* InOuter, const FString& InPath) override;
-	virtual UMediaPlayer* CreateOrUpdateMediaPlayer(UMediaStream* InMediaStream, UMediaPlayer* InCurrentPlayer) override;
+	virtual UMediaPlayer* CreateOrUpdatePlayer(const FMediaStreamSchemeHandlerCreatePlayerParams& InParams) override;
 
 #if WITH_EDITOR
 	virtual void CreatePropertyCustomization(TSharedRef<IPropertyHandle> InPropertyHandle, IDetailGroup& InDetailsGroup, IPropertyTypeCustomizationUtils& InCustomizationUtils) override;
diff --git a/Engine/Plugins/Experimental/MediaStream/Source/MediaStream/Public/IMediaStreamObjectHandler.h b/Engine/Plugins/Experimental/MediaStream/Source/MediaStream/Public/IMediaStreamObjectHandler.h
--- a/Engine/Plugins/Experimental/MediaStream/Source/MediaStream/Public/IMediaStreamObjectHandler.h
+++ b/Engine/Plugins/Experimental/MediaStream/Source/MediaStream/Public/IMediaStreamObjectHandler.h
@@ -4,10 +4,42 @@
 
 #include "Templates/SharedPointer.h"
 
+#include "IMediaStreamObjectHandler.generated.h"
+
 class UMediaPlayer;
 class UMediaStream;
 class UObject;
 
+USTRUCT(BlueprintType)
+struct FMediaStreamObjectHandlerCreatePlayerParams
+{
+	GENERATED_BODY()
+
+	/** The container for the player. */
+	UPROPERTY(BlueprintReadWrite, Category = "Media Stream")
+	TObjectPtr<UMediaStream> MediaStream;
+
+	/** The media source for the player. */
+	UPROPERTY(BlueprintReadWrite, Category = "Media Stream")
+	TObjectPtr<UObject> Source;
+
+	/**
+	 * The current player to update or null.
+	 * If a player is provided, it will be re-used to open the source, if it can be.
+	 * If no player is provided, a new player will be created (if allowed).
+	 */
+	UPROPERTY(BlueprintReadWrite, Category = "Media Stream")
+	TObjectPtr<UMediaPlayer> CurrentPlayer;
+
+	/**
+	 * Whether the new player can open the source or not.
+	 * If this is false, it may mean that a new player is not created or
+	 * an existing player is not updated.
+	 */
+	UPROPERTY(BlueprintReadWrite, Category = "Media Stream")
+	bool bCanOpenSource = true;
+};
+
 /**
  * Implement this interface to add a new object handler.
  * If the derived class has a static UClass* Class then it can be added without specifying it.
@@ -18,11 +50,9 @@ public:
 	virtual ~IMediaStreamObjectHandler() = default;
 
 	/**
-	 * Create a media player for the provided object.
-	 * @param InMediaStream The container for the player.
-	 * @param InObject The object for which to create a player.
-	 * @param InCurrentPlayer The current player to update or null.
+	 * Create or update a UMediaPlayer for the provided source.
+	 * Note: This usually means loading the media source. @see bCanLoadSource.
 	 * @return A player or nullptr.
 	 */
-	virtual UMediaPlayer* CreateOrUpdatePlayer(UMediaStream* InMediaStream, UObject* InObject, UMediaPlayer* InCurrentPlayer) = 0;
+	virtual UMediaPlayer* CreateOrUpdatePlayer(const FMediaStreamObjectHandlerCreatePlayerParams& InParams) = 0;
 };
diff --git a/Engine/Plugins/Experimental/MediaStream/Source/MediaStream/Public/IMediaStreamPlayer.h b/Engine/Plugins/Experimental/MediaStream/Source/MediaStream/Public/IMediaStreamPlayer.h
--- a/Engine/Plugins/Experimental/MediaStream/Source/MediaStream/Public/IMediaStreamPlayer.h
+++ b/Engine/Plugins/Experimental/MediaStream/Source/MediaStream/Public/IMediaStreamPlayer.h
@@ -161,6 +161,13 @@ public:
 	UFUNCTION(BlueprintCallable, Category = "Media Stream|Player")
 	virtual UMediaStream* GetSourceStream() const = 0;
 
+	/**
+	 * Opens the source if it is not already opened.
+	 * @return True if successful.
+	 */
+	UFUNCTION(BlueprintCallable, CallInEditor, Category = "Media Stream|Player")
+	virtual bool OpenSource() = 0;
+
 	/**
 	 * Continues play of the current media.
 	 * @return True if successful.
diff --git a/Engine/Plugins/Experimental/MediaStream/Source/MediaStream/Public/IMediaStreamSchemeHandler.h b/Engine/Plugins/Experimental/MediaStream/Source/MediaStream/Public/IMediaStreamSchemeHandler.h
--- a/Engine/Plugins/Experimental/MediaStream/Source/MediaStream/Public/IMediaStreamSchemeHandler.h
+++ b/Engine/Plugins/Experimental/MediaStream/Source/MediaStream/Public/IMediaStreamSchemeHandler.h
@@ -6,16 +6,47 @@
 
 #include "Containers/ContainersFwd.h"
 
+#include "IMediaStreamSchemeHandler.generated.h"
+
 class IPropertyHandle;
 class UMediaPlayer;
 class UMediaStream;
 struct FMediaStreamSource;
+struct FMediaStreamObjectHandlerCreatePlayerParams;
 
 #if WITH_EDITOR
 class IDetailGroup;
 class IPropertyTypeCustomizationUtils;
 #endif
 
+USTRUCT(BlueprintType)
+struct FMediaStreamSchemeHandlerCreatePlayerParams
+{
+	GENERATED_BODY()
+
+	/** The container for the player. */
+	UPROPERTY(BlueprintReadWrite, Category = "Media Stream")
+	TObjectPtr<UMediaStream> MediaStream;
+
+	/**
+	 * The current player to update or null.
+	 * If a player is provided, it will be re-used to open the source, if it can be.
+	 * If no player is provided, a new player will be created (if allowed).
+	 */
+	UPROPERTY(BlueprintReadWrite, Category = "Media Stream")
+	TObjectPtr<UMediaPlayer> CurrentPlayer;
+
+	/**
+	 * Whether the new player can open the source or not.
+	 * If this is false, it may mean that a new player is not created or
+	 * an existing player is not updated.
+	 */
+	UPROPERTY(BlueprintReadWrite, Category = "Media Stream")
+	bool bCanOpenSource;
+
+	FMediaStreamObjectHandlerCreatePlayerParams operator<<(UObject* InMediaSource) const;
+};
+
 /**
  * Implement this interface to add a new scheme handler.
  * If the derived class has a static TCHAR Scheme[] then it can be added without specifying it.
@@ -34,12 +65,11 @@ public:
 	virtual FMediaStreamSource CreateSource(UObject* InOuter, const FString& InPath) = 0;
 
 	/**
-	 * Creates a media player for the given Media Stream based on its source.
-	 * @param InMediaStream The Media Stream to create a player for.
-	 * @param InCurrentPlayer The current player to update or null.
-	 * @return A media player or nullptr on failure.
+	 * Create or update a UMediaPlayer for the provided source.
+	 * Note: This usually means loading the media source. @see bCanLoadSource.
+	 * @return A player or nullptr.
 	 */
-	virtual UMediaPlayer* CreateOrUpdateMediaPlayer(UMediaStream* InMediaStream, UMediaPlayer* InCurrentPlayer) = 0;
+	virtual UMediaPlayer* CreateOrUpdatePlayer(const FMediaStreamSchemeHandlerCreatePlayerParams& InParams) = 0;
 
 #if WITH_EDITOR
 	virtual void CreatePropertyCustomization(TSharedRef<IPropertyHandle> InPropertyHandle, IDetailGroup& InDetailsGroup, IPropertyTypeCustomizationUtils& InCustomizationUtils) = 0;
diff --git a/Engine/Plugins/Experimental/MediaStream/Source/MediaStream/Public/MediaStreamModule.h b/Engine/Plugins/Experimental/MediaStream/Source/MediaStream/Public/MediaStreamModule.h
--- a/Engine/Plugins/Experimental/MediaStream/Source/MediaStream/Public/MediaStreamModule.h
+++ b/Engine/Plugins/Experimental/MediaStream/Source/MediaStream/Public/MediaStreamModule.h
@@ -8,12 +8,11 @@
 #include "Delegates/DelegateCombinations.h"
 #include "Logging/LogMacros.h"
 
+class IMediaStreamPlayer;
 enum class EMediaStreamPlaybackState : uint8;
 
 DECLARE_LOG_CATEGORY_EXTERN(LogMediaStream, Log, All);
 
-DECLARE_MULTICAST_DELEGATE_OneParam(FMediaStreamEditorPlaybackStateDelegate, EMediaStreamPlaybackState)
-
 /**
  * Media Stream - Content/type agnostic chainable media proxy with media player integration.
  */
@@ -22,13 +21,12 @@ class FMediaStreamModule : public IModuleInterface
 public:
 	MEDIASTREAM_API static FMediaStreamModule& Get();
 
-	MEDIASTREAM_API FMediaStreamEditorPlaybackStateDelegate& GetEditorPlaybackStateEvent();
+	MEDIASTREAM_API bool CanOpenSourceOnLoad();
+
+	MEDIASTREAM_API bool CanAutoplay() const;
 
 	//~ Begin IModuleInterface
 	virtual void StartupModule() override;
 	virtual void ShutdownModule() override;
 	//~ End IModuleInterface
-
-private:
-	FMediaStreamEditorPlaybackStateDelegate EditorPlaybackStateDelegate;
 };
diff --git a/Engine/Plugins/Experimental/MediaStream/Source/MediaStream/Public/MediaStreamObjectHandlerManager.h b/Engine/Plugins/Experimental/MediaStream/Source/MediaStream/Public/MediaStreamObjectHandlerManager.h
--- a/Engine/Plugins/Experimental/MediaStream/Source/MediaStream/Public/MediaStreamObjectHandlerManager.h
+++ b/Engine/Plugins/Experimental/MediaStream/Source/MediaStream/Public/MediaStreamObjectHandlerManager.h
@@ -10,6 +10,7 @@
 class IMediaStreamObjectHandler;
 class UMediaPlayer;
 class UMediaStream;
+struct FMediaStreamObjectHandlerCreatePlayerParams;
 struct FMediaStreamSource;
 
 /**
@@ -184,13 +185,11 @@ public:
 	}
 
 	/**
-	 * Creates a media player for a given Media Stream and source object.
-	 * @param InMediaStream The Media Stream container.
-	 * @param InSourceObject The playback source.
-	 * @param InCurrentPlayer The current player to update or null.
-	 * @return A media player or nullptr on failure.
+	 * Create or update a UMediaPlayer for the provided source.
+	 * Note: This usually means loading the media source. @see bCanLoadSource.
+	 * @return A player or nullptr.
 	 */
-	MEDIASTREAM_API UMediaPlayer* CreateOrUpdateMediaPlayer(UMediaStream* InMediaStream, UObject* InSourceObject, UMediaPlayer* InCurrentPlayer) const;
+	MEDIASTREAM_API UMediaPlayer* CreateOrUpdatePlayer(const FMediaStreamObjectHandlerCreatePlayerParams& InParams) const;
 
 private:
 	TMap<FName, TSharedRef<IMediaStreamObjectHandler>> Handlers;
diff --git a/Engine/Plugins/Experimental/MediaStream/Source/MediaStream/Public/MediaStreamObjectHandlerSubsystem.h b/Engine/Plugins/Experimental/MediaStream/Source/MediaStream/Public/MediaStreamObjectHandlerSubsystem.h
--- a/Engine/Plugins/Experimental/MediaStream/Source/MediaStream/Public/MediaStreamObjectHandlerSubsystem.h
+++ b/Engine/Plugins/Experimental/MediaStream/Source/MediaStream/Public/MediaStreamObjectHandlerSubsystem.h
@@ -8,6 +8,7 @@
 
 class UMediaPlayer;
 class UMediaStream;
+struct FMediaStreamObjectHandlerCreatePlayerParams;
 
 /**
  * Subsystem for blueprint-style interaction with the Object Handler Manager.
@@ -41,12 +42,10 @@ public:
 	MEDIASTREAM_API bool HasObjectHandler(const UClass* InClass) const;
 
 	/**
-	 * Creates a media player for a given Media Stream and source object.
-	 * @param InMediaStream The Media Stream container.
-	 * @param InPlaybackSource The playback source.
-	 * @param InCurrentPlayer The current player to update or null.
-	 * @return A media player or nullptr on failure.
+	 * Create or update a UMediaPlayer for the provided source.
+	 * Note: This usually means loading the media source. @see bCanLoadSource.
+	 * @return A player or nullptr.
 	 */
 	UFUNCTION(BlueprintCallable, Category = "Media Stream|Object")
-	MEDIASTREAM_API UMediaPlayer* CreateMediaPlayer(UMediaStream* InMediaStream, UObject* InPlaybackSource, UMediaPlayer* InCurrentPlayer) const;
+	MEDIASTREAM_API UMediaPlayer* CreateMediaPlayer(const FMediaStreamObjectHandlerCreatePlayerParams& InParams) const;
 };
diff --git a/Engine/Plugins/Experimental/MediaStream/Source/MediaStream/Public/MediaStreamSchemeHandlerManager.h b/Engine/Plugins/Experimental/MediaStream/Source/MediaStream/Public/MediaStreamSchemeHandlerManager.h
--- a/Engine/Plugins/Experimental/MediaStream/Source/MediaStream/Public/MediaStreamSchemeHandlerManager.h
+++ b/Engine/Plugins/Experimental/MediaStream/Source/MediaStream/Public/MediaStreamSchemeHandlerManager.h
@@ -10,6 +10,7 @@
 class IMediaStreamSchemeHandler;
 class UMediaPlayer;
 class UMediaStream;
+struct FMediaStreamSchemeHandlerCreatePlayerParams;
 struct FMediaStreamSource;
 
 /**
@@ -109,12 +110,11 @@ public:
 	MEDIASTREAM_API FMediaStreamSource CreateSource(UObject* InOuter, FName InScheme, const FString& InPath) const;
 
 	/**
-	 * Creates a media player for the given Media Stream based on its source.
-	 * @param InMediaStream The Media Stream to create a player for.
-	 * @param InCurrentPlayer
-	 * @return A media player or nullptr on failure.
+	 * Create or update a UMediaPlayer for the provided source.
+	 * Note: This usually means loading the media source. @see bCanLoadSource.
+	 * @return A player or nullptr.
 	 */
-	MEDIASTREAM_API UMediaPlayer* CreateOrUpdateMediaPlayer(UMediaStream* InMediaStream, UMediaPlayer* InCurrentPlayer) const;
+	MEDIASTREAM_API UMediaPlayer* CreateOrUpdatePlayer(const FMediaStreamSchemeHandlerCreatePlayerParams& InParams) const;
 
 private:
 	TMap<FName, TSharedRef<IMediaStreamSchemeHandler>> Handlers;
diff --git a/Engine/Plugins/Experimental/MediaStream/Source/MediaStream/Public/MediaStreamSchemeHandlerSubsystem.h b/Engine/Plugins/Experimental/MediaStream/Source/MediaStream/Public/MediaStreamSchemeHandlerSubsystem.h
--- a/Engine/Plugins/Experimental/MediaStream/Source/MediaStream/Public/MediaStreamSchemeHandlerSubsystem.h
+++ b/Engine/Plugins/Experimental/MediaStream/Source/MediaStream/Public/MediaStreamSchemeHandlerSubsystem.h
@@ -7,14 +7,11 @@
 #include "HAL/Platform.h"
 #include "UObject/NameTypes.h"
 
-#if WITH_EDITOR
-#include "Templates/SharedPointerFwd.h"
-#endif
-
 #include "MediaStreamSchemeHandlerSubsystem.generated.h"
 
 class UMediaPlayer;
 class UMediaStream;
+struct FMediaStreamSchemeHandlerCreatePlayerParams;
 struct FMediaStreamSource;
 
 /**
@@ -57,11 +54,10 @@ public:
 	FMediaStreamSource CreateSource(UObject* InOuter, FName InScheme, const FString& InPath) const;
 
 	/**
-	 * Creates a media player for the given Media Stream based on its source.
-	 * @param InMediaStream The Media Stream to create a player for.
-	 * @param InCurrentPlayer
-	 * @return A media player or nullptr on failure.
+	 * Create or update a UMediaPlayer for the provided source.
+	 * Note: This usually means loading the media source. @see bCanLoadSource.
+	 * @return A player or nullptr.
 	 */
 	UFUNCTION(BlueprintCallable, Category = "Media Stream|Scheme")
-	MEDIASTREAM_API UMediaPlayer* CreateOrUpdateMediaPlayer(UMediaStream* InMediaStream, UMediaPlayer* InCurrentPlayer) const;
+	MEDIASTREAM_API UMediaPlayer* CreateOrUpdatePlayer(const FMediaStreamSchemeHandlerCreatePlayerParams& InParams) const;
 };
diff --git a/Engine/Plugins/Experimental/MediaStream/Source/MediaStreamEditor/Private/DetailsPanel/MediaStreamCustomization.cpp b/Engine/Plugins/Experimental/MediaStream/Source/MediaStreamEditor/Private/DetailsPanel/MediaStreamCustomization.cpp
--- a/Engine/Plugins/Experimental/MediaStream/Source/MediaStreamEditor/Private/DetailsPanel/MediaStreamCustomization.cpp
+++ b/Engine/Plugins/Experimental/MediaStream/Source/MediaStreamEditor/Private/DetailsPanel/MediaStreamCustomization.cpp
@@ -17,10 +17,100 @@
 #include "Styling/StyleColors.h"
 #include "Widgets/Input/SButton.h"
 #include "Widgets/Images/SImage.h"
+#include "Widgets/Layout/SBox.h"
 #include "Widgets/Text/STextBlock.h"
 
 #define LOCTEXT_NAMESPACE "MediaStreamCustomization"
 
+class SMediaStreamRefreshSliderWidget : public SCompoundWidget
+{
+	SLATE_BEGIN_ARGS(SMediaStreamRefreshSliderWidget) {}
+	SLATE_END_ARGS()
+
+	void Construct(const FArguments& InArgs, UMediaStream* InMediaStream)
+	{
+		MediaStreamWeak = InMediaStream;
+
+		SetCanTick(true);
+
+		ChildSlot
+		[
+			SAssignNew(Content, SBox)
+			.HAlign(HAlign_Fill)
+			.VAlign(VAlign_Fill)
+			[
+				CreateSlider()
+			]
+		];
+	}
+
+	virtual void Tick(const FGeometry& InAllottedGeometry, const double InCurrentTime, const float InDeltaTime) override
+	{
+		SCompoundWidget::Tick(InAllottedGeometry, InCurrentTime, InDeltaTime);
+
+		UMediaPlayer* MediaPlayer = GetMediaPlayer();
+		UMediaPlayer* CurrentMediaPlayer = MediaPlayerWeak.Get();
+
+		if (MediaPlayer != CurrentMediaPlayer)
+		{
+			Content->SetContent(CreateSlider());
+		}
+	}
+
+private:
+	TWeakObjectPtr<UMediaStream> MediaStreamWeak;
+	TWeakObjectPtr<UMediaPlayer> MediaPlayerWeak;
+	TSharedPtr<SBox> Content;
+
+	UMediaPlayer* GetMediaPlayer() const
+	{
+		UMediaStream* MediaStream = MediaStreamWeak.Get();
+
+		if (!MediaStream)
+		{
+			return nullptr;
+		}
+
+		IMediaStreamPlayer* MediaStreamPlayer = MediaStream->GetPlayer().GetInterface();
+
+		if (!MediaStreamPlayer)
+		{
+			return nullptr;
+		}
+
+		return MediaStreamPlayer->GetPlayer();
+	}
+
+	TSharedRef<SWidget> CreateSlider()
+	{
+		UMediaPlayer* MediaPlayer = GetMediaPlayer();
+
+		if (!MediaPlayer)
+		{
+			return SNullWidget::NullWidget;
+		}
+
+		IMediaPlayerEditorModule* MediaPlayerEditorModule = FModuleManager::LoadModulePtr<IMediaPlayerEditorModule>("MediaPlayerEditor");
+
+		if (!MediaPlayerEditorModule)
+		{
+			return SNullWidget::NullWidget;
+		}
+
+		MediaPlayerWeak = MediaPlayer;
+
+		TArray<TWeakObjectPtr<UMediaPlayer>> MediaPlayerList;
+		MediaPlayerList.Add(MediaPlayer);
+
+		const TSharedRef<IMediaPlayerSlider> MediaPlayerSlider = MediaPlayerEditorModule->CreateMediaPlayerSliderWidget(MediaPlayerList);
+
+		MediaPlayerSlider->SetSliderHandleColor(FSlateColor(EStyleColor::AccentBlue));
+		MediaPlayerSlider->SetVisibleWhenInactive(EVisibility::Visible);
+
+		return MediaPlayerSlider;
+	}
+};
+
 FMediaStreamCustomization::FMediaStreamCustomization()
 {
 }
@@ -59,34 +149,145 @@ void FMediaStreamCustomization::AddControlControls(IDetailCategoryBuilder& InMed
 	const ISlateStyle* Stream = &FMediaStreamEditorStyle::Get().Get();
 
 	// Add media player playback slider
-	if (IMediaPlayerEditorModule* MediaPlayerEditorModule = FModuleManager::LoadModulePtr<IMediaPlayerEditorModule>("MediaPlayerEditor"))
-	{
-		const TSharedRef<IMediaPlayerSlider> MediaPlayerSlider =
-			MediaPlayerEditorModule->CreateMediaPlayerSliderWidget(GetMediaPlayers());
+	ControlGroup.AddWidgetRow()
+		.IsEnabled(TAttribute<bool>::CreateSPLambda(this, 
+			[this]()
+			{
+				if (UMediaStream* MediaStream = GetMediaStream())
+				{
+					return !FMediaStreamEditorSequencerLibrary::HasTrack(MediaStream);
+				}
 
-		MediaPlayerSlider->SetSliderHandleColor(FSlateColor(EStyleColor::AccentBlue));
-		MediaPlayerSlider->SetVisibleWhenInactive(EVisibility::Visible);
+				return false;
+			}))
+		[
+			SNew(SMediaStreamRefreshSliderWidget, GetMediaStream())
+		];
+
+	// Add media control buttons.
+	ControlGroup.AddWidgetRow()
+	[
+		SNew(SHorizontalBox)
+
+		// Open button.
+		+ SHorizontalBox::Slot()
+		.VAlign(VAlign_Center)
+		.AutoWidth()
+		[
+			SNew(SButton)
+			.VAlign(VAlign_Center)
+			.Visibility_Lambda([WeakSelf]
+				{
+					if (const TSharedPtr<FMediaStreamCustomization> Self = WeakSelf.Pin())
+					{
+						if (UMediaStream* MediaStream = Self->GetMediaStream())
+						{
+							if (MediaStream->HasValidSource())
+							{
+								if (IMediaStreamPlayer* MediaStreamPlayer = MediaStream->GetPlayer().GetInterface())
+								{
+									if (UMediaPlayer* MediaPlayer = MediaStreamPlayer->GetPlayer())
+									{
+										if (MediaPlayer->IsClosed())
+										{
+											return EVisibility::Visible;
+										}
+
+										return EVisibility::Collapsed;
+									}
+								}
 
-		ControlGroup.AddWidgetRow()
-			.IsEnabled(TAttribute<bool>::CreateSPLambda(this, 
-				[this]()
+								return EVisibility::Visible;
+							}
+						}
+					}
+
+					return EVisibility::Collapsed;
+				})
+			.OnClicked_Lambda([WeakSelf]() -> FReply
 				{
-					if (UMediaStream* MediaStream = GetMediaStream())
+					if (const TSharedPtr<FMediaStreamCustomization> Self = WeakSelf.Pin())
 					{
-						return !FMediaStreamEditorSequencerLibrary::HasTrack(MediaStream);
+						if (UMediaStream* MediaStream = Self->GetMediaStream())
+						{
+							if (IMediaStreamPlayer* MediaStreamPlayer = MediaStream->GetPlayer().GetInterface())
+							{
+								MediaStreamPlayer->OpenSource();
+							}
+						}
+
+						return FReply::Handled();
 					}
 
-					return false;
-				}))
+					return FReply::Unhandled();
+				})
 			[
-				MediaPlayerSlider
-			];
-	}
+				SNew(SImage)
+				.ColorAndOpacity(FSlateColor::UseForeground())
+				.Image(Stream->GetBrush("MediaStreamEditor.OpenMedia.Small"))
+				.ToolTipText(LOCTEXT("Open Media", "Opens the current media, if the source is valid."))
+			]
+		]
+
+		// Close button.
+		+ SHorizontalBox::Slot()
+		.VAlign(VAlign_Center)
+		.AutoWidth()
+		[
+			SNew(SButton)
+			.VAlign(VAlign_Center)
+			.Visibility_Lambda([WeakSelf]
+				{
+					if (const TSharedPtr<FMediaStreamCustomization> Self = WeakSelf.Pin())
+					{
+						if (UMediaStream* MediaStream = Self->GetMediaStream())
+						{
+							if (MediaStream->HasValidSource())
+							{
+								if (IMediaStreamPlayer* MediaStreamPlayer = MediaStream->GetPlayer().GetInterface())
+								{
+									if (UMediaPlayer* MediaPlayer = MediaStreamPlayer->GetPlayer())
+									{
+										if (!FMediaStreamEditorSequencerLibrary::HasTrack(MediaStream))
+										{
+											if (!MediaPlayer->IsClosed())
+											{
+												return EVisibility::Visible;
+											}
+										}
+									}
+								}
+							}
+						}
+					}
+
+					return EVisibility::Collapsed;
+				})
+			.OnClicked_Lambda([WeakSelf]() -> FReply
+				{
+					if (const TSharedPtr<FMediaStreamCustomization> Self = WeakSelf.Pin())
+					{
+						if (UMediaStream* MediaStream = Self->GetMediaStream())
+						{
+							if (IMediaStreamPlayer* MediaStreamPlayer = MediaStream->GetPlayer().GetInterface())
+							{
+								MediaStreamPlayer->Close();
+							}
+						}
+
+						return FReply::Handled();
+					}
+
+					return FReply::Unhandled();
+				})
+			[
+				SNew(SImage)
+				.ColorAndOpacity(FSlateColor::UseForeground())
+				.Image(Stream->GetBrush("MediaStreamEditor.CloseMedia.Small"))
+				.ToolTipText(LOCTEXT("Close Media", "Closes the currently opened media."))
+			]
+		]
 
-	// Add media control buttons.
-	ControlGroup.AddWidgetRow()
-	[
-		SNew(SHorizontalBox)
 		// Rewind button.
 		+ SHorizontalBox::Slot()
 		.VAlign(VAlign_Center)
diff --git a/Engine/Plugins/Experimental/MediaStream/Source/MediaStreamEditor/Private/MediaStreamEditorStyle.cpp b/Engine/Plugins/Experimental/MediaStream/Source/MediaStreamEditor/Private/MediaStreamEditorStyle.cpp
--- a/Engine/Plugins/Experimental/MediaStream/Source/MediaStreamEditor/Private/MediaStreamEditorStyle.cpp
+++ b/Engine/Plugins/Experimental/MediaStream/Source/MediaStreamEditor/Private/MediaStreamEditorStyle.cpp
@@ -20,44 +20,27 @@
 FMediaStreamEditorStyle::FMediaStreamEditorStyle()
 	: FSlateStyleSet("MediaStreamEditorStyle")
 {
-	const FVector2D Icon8x8(8.0f, 8.0f);
 	const FVector2D Icon12x12(12.0f, 12.0f);
 	const FVector2D Icon16x16(16.0f, 16.0f);
 	const FVector2D Icon20x20(20.0f, 20.0f);
 	const FVector2D Icon40x40(40.0f, 40.0f);
 	const FVector2D Icon64x64(64.0f, 64.0f);
 
-	if (TSharedPtr<IPlugin> Plugin = IPluginManager::Get().FindPlugin(TEXT("MediaStream")))
+	if (TSharedPtr<IPlugin> Plugin = IPluginManager::Get().FindPlugin(UE_PLUGIN_NAME))
 	{
 		SetContentRoot(FPaths::Combine(Plugin->GetBaseDir(), TEXT("Resources")));
+
+		// Class icons.
+		Set("ClassIcon.MediaStream", new IMAGE_BRUSH_SVG(TEXT("Icons/MediaStream"), Icon16x16));
+		Set("ClassThumbnail.MediaStream", new IMAGE_BRUSH_SVG(TEXT("Icons/MediaStream"), Icon64x64));
+		Set("ClassIcon.MediaStreamActor", new IMAGE_BRUSH_SVG(TEXT("Icons/MediaStream"), Icon16x16));
+		Set("ClassThumbnail.MediaStreamActor", new IMAGE_BRUSH_SVG(TEXT("Icons/MediaStream"), Icon64x64));
+		Set("ClassIcon.MediaStreamComponent", new IMAGE_BRUSH_SVG(TEXT("Icons/MediaStream"), Icon16x16));
+		Set("ClassThumbnail.MediaStreamComponent", new IMAGE_BRUSH_SVG(TEXT("Icons/MediaStream"), Icon64x64));
 	}
-	// Class icons.
-	Set("ClassIcon.MediaStream", new IMAGE_BRUSH_SVG(TEXT("Icons/MediaStream"), Icon16x16));
-	Set("ClassThumbnail.MediaStream", new IMAGE_BRUSH_SVG(TEXT("Icons/MediaStream"), Icon64x64));
-	Set("ClassIcon.MediaStreamActor", new IMAGE_BRUSH_SVG(TEXT("Icons/MediaStream"), Icon16x16));
-	Set("ClassThumbnail.MediaStreamActor", new IMAGE_BRUSH_SVG(TEXT("Icons/MediaStream"), Icon64x64));
-	Set("ClassIcon.MediaStreamComponent", new IMAGE_BRUSH_SVG(TEXT("Icons/MediaStream"), Icon16x16));
-	Set("ClassThumbnail.MediaStreamComponent", new IMAGE_BRUSH_SVG(TEXT("Icons/MediaStream"), Icon64x64));
 
 	SetContentRoot(FPaths::EnginePluginsDir() / TEXT("Media/MediaPlayerEditor/Content"));
 
-	// buttons
-	Set("MediaPlayerEditor.SourceButton", new IMAGE_BRUSH("btn_source_12x", Icon12x12));
-	Set("MediaPlayerEditor.GoButton", new IMAGE_BRUSH("btn_go_12x", Icon12x12));
-	Set("MediaPlayerEditor.ReloadButton", new IMAGE_BRUSH("btn_reload_12x", Icon12x12));
-	Set("MediaPlayerEditor.SettingsButton", new IMAGE_BRUSH("btn_settings_16x", Icon12x12));
-
-	// misc
-	Set("MediaStreamEditor.DragDropBorder", new BOX_BRUSH("border_dragdrop", 0.5f));
-	Set("MediaStreamEditor.MediaSourceOpened", new IMAGE_BRUSH("mediasource_opened", Icon8x8));
-
-	// tabs
-	Set("MediaStreamEditor.Tabs.Info", new IMAGE_BRUSH("tab_info_16x", Icon16x16));
-	Set("MediaStreamEditor.Tabs.Media", new IMAGE_BRUSH("tab_media_16x", Icon16x16));
-	Set("MediaStreamEditor.Tabs.Player", new IMAGE_BRUSH("tab_player_16x", Icon16x16));
-	Set("MediaStreamEditor.Tabs.Playlist", new IMAGE_BRUSH("tab_playlist_16x", Icon16x16));
-	Set("MediaStreamEditor.Tabs.Stats", new IMAGE_BRUSH("tab_stats_16x", Icon16x16));
-
 	// toolbar icons
 	Set("MediaStreamEditor.CloseMedia", new IMAGE_BRUSH("icon_eject_40x", Icon40x40));
 	Set("MediaStreamEditor.CloseMedia.Small", new IMAGE_BRUSH("icon_eject_40x", Icon20x20));
@@ -90,8 +73,6 @@ FMediaStreamEditorStyle::FMediaStreamEditorStyle()
 		.SetBarThickness(2.0f)
 	);
 
-	Set("MediaStreamEditor.ViewportFont", DEFAULT_FONT("Regular", 18));
-
 	FSlateStyleRegistry::RegisterSlateStyle(*this);
 }
 
diff --git a/Engine/Plugins/Experimental/MediaStream/Source/MediaStreamEditor/Private/MediaStreamEditorToolbar.cpp b/Engine/Plugins/Experimental/MediaStream/Source/MediaStreamEditor/Private/MediaStreamEditorToolbar.cpp
--- a/Engine/Plugins/Experimental/MediaStream/Source/MediaStreamEditor/Private/MediaStreamEditorToolbar.cpp
+++ b/Engine/Plugins/Experimental/MediaStream/Source/MediaStreamEditor/Private/MediaStreamEditorToolbar.cpp
@@ -2,13 +2,17 @@
 
 #include "MediaStreamEditorToolbar.h"
 
+#include "Editor.h"
+#include "IMediaStreamPlayer.h"
+#include "MediaStream.h"
+#include "MediaStreamEditorSequencerLibrary.h"
 #include "MediaStreamEditorStyle.h"
-#include "MediaStreamModule.h"
 #include "Misc/Attribute.h"
 #include "Textures/SlateIcon.h"
 #include "ToolMenu.h"
 #include "ToolMenus.h"
 #include "ToolMenuSection.h"
+#include "UObject/UObjectIterator.h"
 
 #define LOCTEXT_NAMESPACE "MediaStreamEditorToolbar"
 
@@ -26,11 +30,14 @@ FMediaStreamEditorToolbar& FMediaStreamEditorToolbar::Get()
 
 void FMediaStreamEditorToolbar::Integrate()
 {
+	FWorldDelegates::OnWorldCleanup.AddRaw(this, &FMediaStreamEditorToolbar::OnEditorMapCleanup);
 	ExtendLevelEditorToolbar();
 }
 
 void FMediaStreamEditorToolbar::Disintegrate()
 {
+	FWorldDelegates::OnWorldCleanup.RemoveAll(this);
+	UnextendLevelEditorToolbar();
 }
 
 void FMediaStreamEditorToolbar::ExtendLevelEditorToolbar()
@@ -60,7 +67,7 @@ void FMediaStreamEditorToolbar::ExtendLevelEditorToolbar()
 		TEXT("ToggleMediaStream"),
 		FToolUIActionChoice(Action),
 		FText::GetEmpty(),
-		LOCTEXT("MediaStreamToolTip", "Start or stop all Media Streams."),
+		LOCTEXT("MediaStreamToolTip", "Open, start and stop all Media Streams."),
 		TAttribute<FSlateIcon>::CreateRaw(this, &FMediaStreamEditorToolbar::GetIcon)
 	));
 }
@@ -88,6 +95,12 @@ void FMediaStreamEditorToolbar::UnextendLevelEditorToolbar()
 
 void FMediaStreamEditorToolbar::OnToolbarButtonClicked(const FToolMenuContext& InContext)
 {
+	if (!bHasTriggeredLoadAll)
+	{
+		LoadAllSources();
+		return;
+	}
+
 	switch (EditorPlaybackState)
 	{
 		case EMediaStreamPlaybackState::Pause:
@@ -102,14 +115,41 @@ void FMediaStreamEditorToolbar::OnToolbarButtonClicked(const FToolMenuContext& I
 			return;
 	}
 
-	FMediaStreamModule::Get().GetEditorPlaybackStateEvent().Broadcast(EditorPlaybackState);
+	for (UMediaStream* MediaStream : TObjectRange<UMediaStream>())
+	{
+		if (!FMediaStreamEditorSequencerLibrary::HasTrack(MediaStream))
+		{
+			if (IMediaStreamPlayer* MediaStreamPlayer = MediaStream->GetPlayer().GetInterface())
+			{
+				switch (EditorPlaybackState)
+				{
+					case EMediaStreamPlaybackState::Pause:
+						MediaStreamPlayer->Pause();
+						break;
+
+					case EMediaStreamPlaybackState::Play:
+						MediaStreamPlayer->Play();
+						break;
+
+					default:
+						break;
+				}
+			}
+		}
+	}
 }
 
 FSlateIcon FMediaStreamEditorToolbar::GetIcon() const
 {
+	static const FSlateIcon OpenButton = FSlateIcon(FMediaStreamEditorStyle::Get()->GetStyleSetName(), TEXT("MediaStreamEditor.OpenMedia"));
 	static const FSlateIcon PlayButton = FSlateIcon(FMediaStreamEditorStyle::Get()->GetStyleSetName(), TEXT("MediaStreamEditor.PlayMedia"));
 	static const FSlateIcon PauseButton = FSlateIcon(FMediaStreamEditorStyle::Get()->GetStyleSetName(), TEXT("MediaStreamEditor.PauseMedia"));
 
+	if (!bHasTriggeredLoadAll)
+	{
+		return OpenButton;
+	}
+
 	switch (EditorPlaybackState)
 	{
 		case EMediaStreamPlaybackState::Pause:
@@ -123,4 +163,22 @@ FSlateIcon FMediaStreamEditorToolbar::GetIcon() const
 	}
 }
 
+void FMediaStreamEditorToolbar::LoadAllSources()
+{
+	bHasTriggeredLoadAll = true;
+
+	for (UMediaStream* MediaStream : TObjectRange<UMediaStream>())
+	{
+		if (IMediaStreamPlayer* MediaStreamPlayer = MediaStream->GetPlayer().GetInterface())
+		{
+			MediaStreamPlayer->OpenSource();
+		}
+	}
+}
+
+void FMediaStreamEditorToolbar::OnEditorMapCleanup(UWorld* InWorld, bool bInSessionEnded, bool bInCleanupResources)
+{
+	bHasTriggeredLoadAll = false;
+}
+
 #undef LOCTEXT_NAMESPACE
diff --git a/Engine/Plugins/Experimental/MediaStream/Source/MediaStreamEditor/Private/MediaStreamEditorToolbar.h b/Engine/Plugins/Experimental/MediaStream/Source/MediaStreamEditor/Private/MediaStreamEditorToolbar.h
--- a/Engine/Plugins/Experimental/MediaStream/Source/MediaStreamEditor/Private/MediaStreamEditorToolbar.h
+++ b/Engine/Plugins/Experimental/MediaStream/Source/MediaStreamEditor/Private/MediaStreamEditorToolbar.h
@@ -18,6 +18,7 @@ public:
 	void Disintegrate();
 
 private:
+	bool bHasTriggeredLoadAll = false;
 	EMediaStreamPlaybackState EditorPlaybackState = EMediaStreamPlaybackState::Pause;
 
 	void ExtendLevelEditorToolbar();
@@ -27,4 +28,8 @@ private:
 	void OnToolbarButtonClicked(const FToolMenuContext& InContext);
 
 	FSlateIcon GetIcon() const;
+
+	void LoadAllSources();
+
+	void OnEditorMapCleanup(UWorld* InWorld, bool bInSessionEnded, bool bInCleanupResources);
 };
