Media Stream / Material Designer: Details update
- Change to the media stream to remove the necessity to refresh the details panel

diff --git a/Engine/Plugins/Experimental/DynamicMaterialMediaStreamBridge/Source/DynamicMaterialMediaStreamBridge/Private/DMMaterialValueMediaStream.cpp b/Engine/Plugins/Experimental/DynamicMaterialMediaStreamBridge/Source/DynamicMaterialMediaStreamBridge/Private/DMMaterialValueMediaStream.cpp
--- a/Engine/Plugins/Experimental/DynamicMaterialMediaStreamBridge/Source/DynamicMaterialMediaStreamBridge/Private/DMMaterialValueMediaStream.cpp
+++ b/Engine/Plugins/Experimental/DynamicMaterialMediaStreamBridge/Source/DynamicMaterialMediaStreamBridge/Private/DMMaterialValueMediaStream.cpp
@@ -206,8 +206,6 @@ void UDMMaterialValueMediaStream::UnsubscribeFromPlayerChange()
 void UDMMaterialValueMediaStream::OnSourceChanged(UMediaStream* InMediaStream)
 {
 	UpdatePlayer();
-
-	Update(this, EDMUpdateType::RefreshDetailView);
 }
 
 void UDMMaterialValueMediaStream::UpdatePlayer()
diff --git a/Engine/Plugins/Experimental/MediaStream/Source/MediaStream/Private/SchemeHandlers/MediaStreamAssetSchemeHandler.cpp b/Engine/Plugins/Experimental/MediaStream/Source/MediaStream/Private/SchemeHandlers/MediaStreamAssetSchemeHandler.cpp
--- a/Engine/Plugins/Experimental/MediaStream/Source/MediaStream/Private/SchemeHandlers/MediaStreamAssetSchemeHandler.cpp
+++ b/Engine/Plugins/Experimental/MediaStream/Source/MediaStream/Private/SchemeHandlers/MediaStreamAssetSchemeHandler.cpp
@@ -125,18 +125,19 @@ void FMediaStreamAssetSchemeHandler::AddAssetSelector(TSharedRef<IPropertyHandle
 		return;
 	}
 
-	FDetailWidgetRow& FilePathRow = InSourceGroup.AddWidgetRow();
+	FDetailWidgetRow& AssetSelectRow = InSourceGroup.AddWidgetRow();
 
-	FilePathRow.Visibility(TAttribute<EVisibility>::CreateSP(this, &FMediaStreamAssetSchemeHandler::GetAssetSelectorVisibility, PropertyHandleWeak));
+	TAttribute<EVisibility> Visibility = TAttribute<EVisibility>::CreateSP(this, &FMediaStreamAssetSchemeHandler::GetAssetSelectorVisibility, PropertyHandleWeak);
 
-	FilePathRow.NameContent()
+	AssetSelectRow.NameContent()
 	[
 		SNew(STextBlock)
 		.Font(IDetailLayoutBuilder::GetDetailFont())
 		.Text(LOCTEXT("Asset", "Asset"))
+		.Visibility(Visibility)
 	];
 
-	FilePathRow.ValueContent()
+	AssetSelectRow.ValueContent()
 	[
 		SNew(SObjectPropertyEntryBox)
 		.ObjectPath_Static(&FMediaStreamSchemeHandlerLibrary::GetPath, PropertyHandleWeak)
@@ -145,6 +146,7 @@ void FMediaStreamAssetSchemeHandler::AddAssetSelector(TSharedRef<IPropertyHandle
 		.DisplayThumbnail(true)
 		.AllowedClass(UMediaSource::StaticClass())
 		.PropertyHandle(InPropertyHandle)
+		.Visibility(Visibility)
 	];
 }
 
diff --git a/Engine/Plugins/Experimental/MediaStream/Source/MediaStream/Private/SchemeHandlers/MediaStreamFileSchemeHandler.cpp b/Engine/Plugins/Experimental/MediaStream/Source/MediaStream/Private/SchemeHandlers/MediaStreamFileSchemeHandler.cpp
--- a/Engine/Plugins/Experimental/MediaStream/Source/MediaStream/Private/SchemeHandlers/MediaStreamFileSchemeHandler.cpp
+++ b/Engine/Plugins/Experimental/MediaStream/Source/MediaStream/Private/SchemeHandlers/MediaStreamFileSchemeHandler.cpp
@@ -140,13 +140,14 @@ void FMediaStreamFileSchemeHandler::AddFileSelector(TSharedRef<IPropertyHandle>
 
 	FDetailWidgetRow& FilePathRow = InSourceGroup.AddWidgetRow();
 
-	FilePathRow.Visibility(TAttribute<EVisibility>::CreateSP(this, &FMediaStreamFileSchemeHandler::GetFileSelectorVisibility, PropertyHandleWeak));
+	TAttribute<EVisibility> Visibility = TAttribute<EVisibility>::CreateSP(this, &FMediaStreamFileSchemeHandler::GetFileSelectorVisibility, PropertyHandleWeak);
 
 	FilePathRow.NameContent()
 	[
 		SNew(STextBlock)
 		.Font(IDetailLayoutBuilder::GetDetailFont())
 		.Text(LOCTEXT("FilePath", "File Path"))
+		.Visibility(Visibility)
 	];
 
 	FilePathRow.ValueContent()
@@ -159,6 +160,7 @@ void FMediaStreamFileSchemeHandler::AddFileSelector(TSharedRef<IPropertyHandle>
 		.FilePath_Static(&FMediaStreamSchemeHandlerLibrary::GetPath, PropertyHandleWeak)
 		.FileTypeFilter(TEXT("All files (*.*)|*.*"))
 		.OnPathPicked(this, &FMediaStreamFileSchemeHandler::OnFilePicked, PropertyHandleWeak)
+		.Visibility(Visibility)
 	];
 }
 
