Custom Details View / Media Stream: Fixed row visibility.
- CDV now properly uses the row visibility attr to produce dynamic details panels.
- The media stream customisations now correctly set the row visibility.

diff --git a/Engine/Plugins/Experimental/CustomDetailsView/Source/CustomDetailsView/Private/Items/CustomDetailsViewDetailTreeNodeItem.cpp b/Engine/Plugins/Experimental/CustomDetailsView/Source/CustomDetailsView/Private/Items/CustomDetailsViewDetailTreeNodeItem.cpp
--- a/Engine/Plugins/Experimental/CustomDetailsView/Source/CustomDetailsView/Private/Items/CustomDetailsViewDetailTreeNodeItem.cpp
+++ b/Engine/Plugins/Experimental/CustomDetailsView/Source/CustomDetailsView/Private/Items/CustomDetailsViewDetailTreeNodeItem.cpp
@@ -365,7 +365,7 @@ void FCustomDetailsViewDetailTreeNodeItem::GenerateCustomChildren(const TSharedR
 		: ECustomDetailsViewNodePropertyFlag::None;
 
 	TArray<TSharedRef<IDetailTreeNode>> NodeChildren;
-	DetailTreeNode->GetChildren(NodeChildren);
+	DetailTreeNode->GetChildren(NodeChildren, /* Ignore Visibility */ true);
 
 	AddChildDetailsTreeNodes(InParentItem, ChildNodePropertyFlags, NodeChildren, OutChildren);
 }
diff --git a/Engine/Plugins/Experimental/CustomDetailsView/Source/CustomDetailsView/Private/Items/CustomDetailsViewItemBase.cpp b/Engine/Plugins/Experimental/CustomDetailsView/Source/CustomDetailsView/Private/Items/CustomDetailsViewItemBase.cpp
--- a/Engine/Plugins/Experimental/CustomDetailsView/Source/CustomDetailsView/Private/Items/CustomDetailsViewItemBase.cpp
+++ b/Engine/Plugins/Experimental/CustomDetailsView/Source/CustomDetailsView/Private/Items/CustomDetailsViewItemBase.cpp
@@ -349,6 +349,11 @@ bool FCustomDetailsViewItemBase::IsWidgetVisible() const
 		}
 	}
 
+	if (DetailWidgetRow.VisibilityAttr.Get(/* Default Value */ EVisibility::Visible) != EVisibility::Visible)
+	{
+		return false;
+	}
+
 	if (DetailWidgetRow.HasColumns())
 	{
 		return (DetailWidgetRow.HasNameContent() && DetailWidgetRow.NameWidget.Widget->GetVisibility().IsVisible())
@@ -368,6 +373,11 @@ void FCustomDetailsViewItemBase::SetEnabledOverride(TAttribute<bool> InOverride)
 	EnabledOverride = InOverride;
 }
 
+const FDetailWidgetRow& FCustomDetailsViewItemBase::GetDetailWidgetRow() const
+{
+	return DetailWidgetRow;
+}
+
 void FCustomDetailsViewItemBase::AddWholeRowWidget(const TSharedRef<SSplitter>& InSplitter, const TSharedPtr<SWidget>& InPrependWidget,
 	const FDetailColumnSizeData& InColumnSizeData, const FMargin& InPadding)
 {
diff --git a/Engine/Plugins/Experimental/CustomDetailsView/Source/CustomDetailsView/Private/Items/CustomDetailsViewItemBase.h b/Engine/Plugins/Experimental/CustomDetailsView/Source/CustomDetailsView/Private/Items/CustomDetailsViewItemBase.h
--- a/Engine/Plugins/Experimental/CustomDetailsView/Source/CustomDetailsView/Private/Items/CustomDetailsViewItemBase.h
+++ b/Engine/Plugins/Experimental/CustomDetailsView/Source/CustomDetailsView/Private/Items/CustomDetailsViewItemBase.h
@@ -54,6 +54,7 @@ public:
 	virtual bool IsWidgetVisible() const override;
 	virtual void SetValueWidgetWidthOverride(TOptional<float> InWidth) override;
 	virtual void SetEnabledOverride(TAttribute<bool> InOverride) override;
+	virtual const FDetailWidgetRow& GetDetailWidgetRow() const override;
 	//~ End ICustomDetailsViewItem
 
 	virtual void AddWholeRowWidget(const TSharedRef<SSplitter>& InSplitter
diff --git a/Engine/Plugins/Experimental/CustomDetailsView/Source/CustomDetailsView/Private/Slate/SCustomDetailsViewItemRow.cpp b/Engine/Plugins/Experimental/CustomDetailsView/Source/CustomDetailsView/Private/Slate/SCustomDetailsViewItemRow.cpp
--- a/Engine/Plugins/Experimental/CustomDetailsView/Source/CustomDetailsView/Private/Slate/SCustomDetailsViewItemRow.cpp
+++ b/Engine/Plugins/Experimental/CustomDetailsView/Source/CustomDetailsView/Private/Slate/SCustomDetailsViewItemRow.cpp
@@ -1,6 +1,8 @@
 // Copyright Epic Games, Inc. All Rights Reserved.
 
 #include "SCustomDetailsViewItemRow.h"
+
+#include "DetailWidgetRow.h"
 #include "Items/ICustomDetailsViewItem.h"
 #include "Widgets/Views/STreeView.h"
 
@@ -18,9 +20,16 @@ void SCustomDetailsViewItemRow::Construct(const FArguments& InArgs
 
 	ExpanderArrowWidget = ExpanderArrow;
 
+	TSharedRef<SWidget> Child = InItem->MakeWidget(ExpanderArrow, SharedThis(this));
+
+	if (InItem->GetDetailWidgetRow().VisibilityAttr.IsSet())
+	{
+		Child->SetVisibility(InItem->GetDetailWidgetRow().VisibilityAttr);
+	}
+
 	STableRow<TSharedPtr<ICustomDetailsViewItem>>::Construct(STableRow<TSharedPtr<ICustomDetailsViewItem>>::FArguments(InArgs._TableRowArgs)
 		[
-			InItem->MakeWidget(ExpanderArrow, SharedThis(this))
+			Child
 		]
 		, InOwnerTable);
 
diff --git a/Engine/Plugins/Experimental/CustomDetailsView/Source/CustomDetailsView/Public/Items/ICustomDetailsViewItem.h b/Engine/Plugins/Experimental/CustomDetailsView/Source/CustomDetailsView/Public/Items/ICustomDetailsViewItem.h
--- a/Engine/Plugins/Experimental/CustomDetailsView/Source/CustomDetailsView/Public/Items/ICustomDetailsViewItem.h
+++ b/Engine/Plugins/Experimental/CustomDetailsView/Source/CustomDetailsView/Public/Items/ICustomDetailsViewItem.h
@@ -85,6 +85,9 @@ public:
 	/** Overrides the created widget's enabled status. */
 	virtual void SetEnabledOverride(TAttribute<bool> InOverride) = 0;
 
+	/** Returns the widget row that stores the default of this widget. */
+	virtual const FDetailWidgetRow& GetDetailWidgetRow() const = 0;
+
 protected:
 	/** Adds the children of this node. */
 	virtual void GatherChildren(const TSharedRef<ICustomDetailsViewItem>& InParentItem,
diff --git a/Engine/Plugins/Experimental/MediaStream/Source/MediaStream/Private/SchemeHandlers/MediaStreamAssetSchemeHandler.cpp b/Engine/Plugins/Experimental/MediaStream/Source/MediaStream/Private/SchemeHandlers/MediaStreamAssetSchemeHandler.cpp
--- a/Engine/Plugins/Experimental/MediaStream/Source/MediaStream/Private/SchemeHandlers/MediaStreamAssetSchemeHandler.cpp
+++ b/Engine/Plugins/Experimental/MediaStream/Source/MediaStream/Private/SchemeHandlers/MediaStreamAssetSchemeHandler.cpp
@@ -127,14 +127,13 @@ void FMediaStreamAssetSchemeHandler::AddAssetSelector(TSharedRef<IPropertyHandle
 
 	FDetailWidgetRow& AssetSelectRow = InSourceGroup.AddWidgetRow();
 
-	TAttribute<EVisibility> Visibility = TAttribute<EVisibility>::CreateSP(this, &FMediaStreamAssetSchemeHandler::GetAssetSelectorVisibility, PropertyHandleWeak);
+	AssetSelectRow.VisibilityAttr = TAttribute<EVisibility>::CreateSP(this, &FMediaStreamAssetSchemeHandler::GetAssetSelectorVisibility, PropertyHandleWeak);
 
 	AssetSelectRow.NameContent()
 	[
 		SNew(STextBlock)
 		.Font(IDetailLayoutBuilder::GetDetailFont())
 		.Text(LOCTEXT("Asset", "Asset"))
-		.Visibility(Visibility)
 	];
 
 	AssetSelectRow.ValueContent()
@@ -146,7 +145,6 @@ void FMediaStreamAssetSchemeHandler::AddAssetSelector(TSharedRef<IPropertyHandle
 		.DisplayThumbnail(true)
 		.AllowedClass(UMediaSource::StaticClass())
 		.PropertyHandle(InPropertyHandle)
-		.Visibility(Visibility)
 	];
 }
 
diff --git a/Engine/Plugins/Experimental/MediaStream/Source/MediaStream/Private/SchemeHandlers/MediaStreamFileSchemeHandler.cpp b/Engine/Plugins/Experimental/MediaStream/Source/MediaStream/Private/SchemeHandlers/MediaStreamFileSchemeHandler.cpp
--- a/Engine/Plugins/Experimental/MediaStream/Source/MediaStream/Private/SchemeHandlers/MediaStreamFileSchemeHandler.cpp
+++ b/Engine/Plugins/Experimental/MediaStream/Source/MediaStream/Private/SchemeHandlers/MediaStreamFileSchemeHandler.cpp
@@ -140,14 +140,13 @@ void FMediaStreamFileSchemeHandler::AddFileSelector(TSharedRef<IPropertyHandle>
 
 	FDetailWidgetRow& FilePathRow = InSourceGroup.AddWidgetRow();
 
-	TAttribute<EVisibility> Visibility = TAttribute<EVisibility>::CreateSP(this, &FMediaStreamFileSchemeHandler::GetFileSelectorVisibility, PropertyHandleWeak);
+	FilePathRow.VisibilityAttr = TAttribute<EVisibility>::CreateSP(this, &FMediaStreamFileSchemeHandler::GetFileSelectorVisibility, PropertyHandleWeak);
 
 	FilePathRow.NameContent()
 	[
 		SNew(STextBlock)
 		.Font(IDetailLayoutBuilder::GetDetailFont())
 		.Text(LOCTEXT("FilePath", "File Path"))
-		.Visibility(Visibility)
 	];
 
 	FilePathRow.ValueContent()
@@ -160,7 +159,6 @@ void FMediaStreamFileSchemeHandler::AddFileSelector(TSharedRef<IPropertyHandle>
 		.FilePath_Static(&FMediaStreamSchemeHandlerLibrary::GetPath, PropertyHandleWeak)
 		.FileTypeFilter(TEXT("All files (*.*)|*.*"))
 		.OnPathPicked(this, &FMediaStreamFileSchemeHandler::OnFilePicked, PropertyHandleWeak)
-		.Visibility(Visibility)
 	];
 }
 
