Media Stream / Material Designer: Details panel clean up
- Media stream object, component and actor: Details panel cleaned up.
- Material Designer: Details panel cleaned up.
- Fixed source not being selectable.

diff --git a/Engine/Plugins/Experimental/DynamicMaterial/Source/DynamicMaterialEditor/Private/UI/Widgets/Editor/SDMObjectEditorWidgetBase.cpp b/Engine/Plugins/Experimental/DynamicMaterial/Source/DynamicMaterialEditor/Private/UI/Widgets/Editor/SDMObjectEditorWidgetBase.cpp
--- a/Engine/Plugins/Experimental/DynamicMaterial/Source/DynamicMaterialEditor/Private/UI/Widgets/Editor/SDMObjectEditorWidgetBase.cpp
+++ b/Engine/Plugins/Experimental/DynamicMaterial/Source/DynamicMaterialEditor/Private/UI/Widgets/Editor/SDMObjectEditorWidgetBase.cpp
@@ -98,7 +98,7 @@ TSharedRef<SWidget> SDMObjectEditorWidgetBase::CreateWidget()
 
 	for (const FDMPropertyHandle& PropertyRow : PropertyRows)
 	{
-		const bool bHasValidCustomWidget = PropertyRow.ValueWidget.IsValid() && !PropertyRow.ValueName.IsNone() && PropertyRow.NameOverride.IsSet();
+		const bool bHasValidCustomWidget = PropertyRow.ValueWidget.IsValid() && !PropertyRow.ValueName.IsNone();
 
 		if (!PropertyRow.PreviewHandle.DetailTreeNode && !bHasValidCustomWidget)
 		{
@@ -266,10 +266,13 @@ void SDMObjectEditorWidgetBase::AddCustomRow(const TSharedRef<ICustomDetailsView
 	const TSharedRef<ICustomDetailsViewItem>& InParent, ECustomDetailsTreeInsertPosition InPosition, 
 	const FDMPropertyHandle& InPropertyRow)
 {
+	const FText RowDisplayName = InPropertyRow.NameOverride.Get(FText::GetEmpty());
+	const bool bIsWholeRow = !RowDisplayName.IsEmpty();
+
 	TSharedPtr<ICustomDetailsViewCustomItem> Item = InDetailsView->CreateCustomItem(
 		InParent,
 		InPropertyRow.ValueName,
-		InPropertyRow.NameOverride.GetValue(),
+		RowDisplayName,
 		InPropertyRow.NameToolTipOverride.Get(FText::GetEmpty())
 	);
 
@@ -278,24 +281,31 @@ void SDMObjectEditorWidgetBase::AddCustomRow(const TSharedRef<ICustomDetailsView
 		return;
 	}
 
-	Item->SetValueWidget(InPropertyRow.ValueWidget.ToSharedRef());
-
-	if (!InPropertyRow.bEnabled)
+	if (bIsWholeRow)
 	{
-		Item->AsItem()->SetEnabledOverride(false);
-
-		// Disable the expansion widgets (SNullWidget is treated as removing the override).
-		Item->SetExpansionWidget(SNew(SBox));
+		Item->SetWholeRowWidget(InPropertyRow.ValueWidget.ToSharedRef());
 	}
-
-	if (InPropertyRow.MaxWidth.IsSet())
+	else
 	{
-		Item->AsItem()->SetValueWidgetWidthOverride(InPropertyRow.MaxWidth);
-	}
+		Item->SetValueWidget(InPropertyRow.ValueWidget.ToSharedRef());
 
-	if (InPropertyRow.PreviewHandle.PropertyHandle.IsValid())
-	{
-		AddDetailTreeRowExtensionWidgets(InDetailsView, InPropertyRow, Item->AsItem(), InPropertyRow.PreviewHandle.PropertyHandle.ToSharedRef());
+		if (!InPropertyRow.bEnabled)
+		{
+			Item->AsItem()->SetEnabledOverride(false);
+
+			// Disable the expansion widgets (SNullWidget is treated as removing the override).
+			Item->SetExpansionWidget(SNew(SBox));
+		}
+
+		if (InPropertyRow.MaxWidth.IsSet())
+		{
+			Item->AsItem()->SetValueWidgetWidthOverride(InPropertyRow.MaxWidth);
+		}
+
+		if (InPropertyRow.PreviewHandle.PropertyHandle.IsValid())
+		{
+			AddDetailTreeRowExtensionWidgets(InDetailsView, InPropertyRow, Item->AsItem(), InPropertyRow.PreviewHandle.PropertyHandle.ToSharedRef());
+		}
 	}
 
 	InDetailsView->ExtendTree(InParent->GetItemId(), InPosition, Item->AsItem());
diff --git a/Engine/Plugins/Experimental/DynamicMaterialMediaStreamBridge/Source/DynamicMaterialMediaStreamBridge/DynamicMaterialMediaStreamBridge.Build.cs b/Engine/Plugins/Experimental/DynamicMaterialMediaStreamBridge/Source/DynamicMaterialMediaStreamBridge/DynamicMaterialMediaStreamBridge.Build.cs
--- a/Engine/Plugins/Experimental/DynamicMaterialMediaStreamBridge/Source/DynamicMaterialMediaStreamBridge/DynamicMaterialMediaStreamBridge.Build.cs
+++ b/Engine/Plugins/Experimental/DynamicMaterialMediaStreamBridge/Source/DynamicMaterialMediaStreamBridge/DynamicMaterialMediaStreamBridge.Build.cs
@@ -25,6 +25,14 @@ public class DynamicMaterialMediaStreamBridge : ModuleRules
 			}
 		);
 
+		if (Target.Type == TargetRules.TargetType.Editor)
+		{
+			PrivateDependencyModuleNames.AddRange(new string[]
+			{
+				"SlateCore"
+			});
+		}
+
 		ShortName = "DMMSBridge";
 	}
 }
diff --git a/Engine/Plugins/Experimental/DynamicMaterialMediaStreamBridge/Source/DynamicMaterialMediaStreamBridge/Private/DMMaterialValueMediaStream.cpp b/Engine/Plugins/Experimental/DynamicMaterialMediaStreamBridge/Source/DynamicMaterialMediaStreamBridge/Private/DMMaterialValueMediaStream.cpp
--- a/Engine/Plugins/Experimental/DynamicMaterialMediaStreamBridge/Source/DynamicMaterialMediaStreamBridge/Private/DMMaterialValueMediaStream.cpp
+++ b/Engine/Plugins/Experimental/DynamicMaterialMediaStreamBridge/Source/DynamicMaterialMediaStreamBridge/Private/DMMaterialValueMediaStream.cpp
@@ -9,6 +9,10 @@
 #include "UObject/Package.h"
 #include "UObject/UObjectGlobals.h"
 
+#if WITH_EDITOR
+#include "Styling/SlateIconFinder.h"
+#endif
+
 #define LOCTEXT_NAMESPACE "DMMaterialValueMediaStream"
 
 namespace UE::DynamicMaterialMediaStreamBridge::Private
@@ -171,6 +175,11 @@ void UDMMaterialValueMediaStream::PostEditorDuplicate(UDynamicMaterialModel* InM
 	UpdatePlayer();
 }
 
+FSlateIcon UDMMaterialValueMediaStream::GetComponentIcon() const
+{
+	return FSlateIconFinder::FindIcon("ClassIcon.MediaPlayer");
+}
+
 void UDMMaterialValueMediaStream::PostEditChangeProperty(FPropertyChangedEvent& InPropertyChangedEvent)
 {
 	Super::PostEditChangeProperty(InPropertyChangedEvent);
diff --git a/Engine/Plugins/Experimental/DynamicMaterialMediaStreamBridge/Source/DynamicMaterialMediaStreamBridge/Private/DMMaterialValueMediaStreamDynamic.cpp b/Engine/Plugins/Experimental/DynamicMaterialMediaStreamBridge/Source/DynamicMaterialMediaStreamBridge/Private/DMMaterialValueMediaStreamDynamic.cpp
--- a/Engine/Plugins/Experimental/DynamicMaterialMediaStreamBridge/Source/DynamicMaterialMediaStreamBridge/Private/DMMaterialValueMediaStreamDynamic.cpp
+++ b/Engine/Plugins/Experimental/DynamicMaterialMediaStreamBridge/Source/DynamicMaterialMediaStreamBridge/Private/DMMaterialValueMediaStreamDynamic.cpp
@@ -6,6 +6,10 @@
 #include "IMediaStreamPlayer.h"
 #include "MediaStream.h"
 
+#if WITH_EDITOR
+#include "Styling/SlateIconFinder.h"
+#endif
+
 #define LOCTEXT_NAMESPACE "DMMaterialValueMediaStreamDynamic"
 
 UDMMaterialValueMediaStreamDynamic::UDMMaterialValueMediaStreamDynamic()
@@ -122,6 +126,11 @@ void UDMMaterialValueMediaStreamDynamic::PostEditorDuplicate(UDynamicMaterialMod
 	UpdatePlayer();
 }
 
+FSlateIcon UDMMaterialValueMediaStreamDynamic::GetComponentIcon() const
+{
+	return FSlateIconFinder::FindIcon("ClassIcon.MediaPlayer");
+}
+
 void UDMMaterialValueMediaStreamDynamic::PostEditChangeProperty(FPropertyChangedEvent& InPropertyChangedEvent)
 {
 	Super::PostEditChangeProperty(InPropertyChangedEvent);
diff --git a/Engine/Plugins/Experimental/DynamicMaterialMediaStreamBridge/Source/DynamicMaterialMediaStreamBridge/Public/DMMaterialValueMediaStream.h b/Engine/Plugins/Experimental/DynamicMaterialMediaStreamBridge/Source/DynamicMaterialMediaStreamBridge/Public/DMMaterialValueMediaStream.h
--- a/Engine/Plugins/Experimental/DynamicMaterialMediaStreamBridge/Source/DynamicMaterialMediaStreamBridge/Public/DMMaterialValueMediaStream.h
+++ b/Engine/Plugins/Experimental/DynamicMaterialMediaStreamBridge/Source/DynamicMaterialMediaStreamBridge/Public/DMMaterialValueMediaStream.h
@@ -41,6 +41,7 @@ public:
 	DYNAMICMATERIALMEDIASTREAMBRIDGE_API virtual FString GetComponentPathComponent() const override;
 	DYNAMICMATERIALMEDIASTREAMBRIDGE_API virtual FText GetComponentDescription() const override;
 	DYNAMICMATERIALMEDIASTREAMBRIDGE_API virtual void PostEditorDuplicate(UDynamicMaterialModel* InMaterialModel, UDMMaterialComponent* InParent) override;
+	DYNAMICMATERIALMEDIASTREAMBRIDGE_API virtual FSlateIcon GetComponentIcon() const override;
 	//~ End UDMMaterialComponent
 
 	//~ Begin UObject
diff --git a/Engine/Plugins/Experimental/DynamicMaterialMediaStreamBridge/Source/DynamicMaterialMediaStreamBridge/Public/DMMaterialValueMediaStreamDynamic.h b/Engine/Plugins/Experimental/DynamicMaterialMediaStreamBridge/Source/DynamicMaterialMediaStreamBridge/Public/DMMaterialValueMediaStreamDynamic.h
--- a/Engine/Plugins/Experimental/DynamicMaterialMediaStreamBridge/Source/DynamicMaterialMediaStreamBridge/Public/DMMaterialValueMediaStreamDynamic.h
+++ b/Engine/Plugins/Experimental/DynamicMaterialMediaStreamBridge/Source/DynamicMaterialMediaStreamBridge/Public/DMMaterialValueMediaStreamDynamic.h
@@ -37,6 +37,7 @@ public:
 	DYNAMICMATERIALMEDIASTREAMBRIDGE_API virtual FString GetComponentPathComponent() const override;
 	DYNAMICMATERIALMEDIASTREAMBRIDGE_API virtual FText GetComponentDescription() const override;
 	DYNAMICMATERIALMEDIASTREAMBRIDGE_API virtual void PostEditorDuplicate(UDynamicMaterialModel* InMaterialModel, UDMMaterialComponent* InParent) override;
+	DYNAMICMATERIALMEDIASTREAMBRIDGE_API virtual FSlateIcon GetComponentIcon() const override;
 	//~ End UDMMaterialComponent
 
 	//~ Begin UObject
diff --git a/Engine/Plugins/Experimental/DynamicMaterialMediaStreamBridge/Source/DynamicMaterialMediaStreamBridgeEditor/DynamicMaterialMediaStreamBridgeEditor.Build.cs b/Engine/Plugins/Experimental/DynamicMaterialMediaStreamBridge/Source/DynamicMaterialMediaStreamBridgeEditor/DynamicMaterialMediaStreamBridgeEditor.Build.cs
--- a/Engine/Plugins/Experimental/DynamicMaterialMediaStreamBridge/Source/DynamicMaterialMediaStreamBridgeEditor/DynamicMaterialMediaStreamBridgeEditor.Build.cs
+++ b/Engine/Plugins/Experimental/DynamicMaterialMediaStreamBridge/Source/DynamicMaterialMediaStreamBridgeEditor/DynamicMaterialMediaStreamBridgeEditor.Build.cs
@@ -24,6 +24,7 @@ public class DynamicMaterialMediaStreamBridgeEditor : ModuleRules
 				"DynamicMaterialEditor",
 				"DynamicMaterialMediaStreamBridge",
 				"MediaStream",
+				"MediaStreamEditor",
 				"Slate",
 				"SlateCore",
 				"ToolMenus",
diff --git a/Engine/Plugins/Experimental/DynamicMaterialMediaStreamBridge/Source/DynamicMaterialMediaStreamBridgeEditor/Private/DMMaterialValueMediaStreamPropertyRowGenerator.cpp b/Engine/Plugins/Experimental/DynamicMaterialMediaStreamBridge/Source/DynamicMaterialMediaStreamBridgeEditor/Private/DMMaterialValueMediaStreamPropertyRowGenerator.cpp
--- a/Engine/Plugins/Experimental/DynamicMaterialMediaStreamBridge/Source/DynamicMaterialMediaStreamBridgeEditor/Private/DMMaterialValueMediaStreamPropertyRowGenerator.cpp
+++ b/Engine/Plugins/Experimental/DynamicMaterialMediaStreamBridge/Source/DynamicMaterialMediaStreamBridgeEditor/Private/DMMaterialValueMediaStreamPropertyRowGenerator.cpp
@@ -8,11 +8,14 @@
 #include "IDynamicMaterialEditorModule.h"
 #include "IMediaStreamPlayer.h"
 #include "MediaStream.h"
+#include "MediaStreamWidgets.h"
 #include "Model/DynamicMaterialModelBase.h"
 #include "UI/Utils/IDMWidgetLibrary.h"
 #include "UObject/Object.h"
 #include "Utils/DMMaterialModelFunctionLibrary.h"
 
+#define LOCTEXT_NAMESPACE "DMMaterialValueMediaStreamPropertyRowGenerator"
+
 const TSharedRef<FDMMaterialValueMediaStreamPropertyRowGenerator>& FDMMaterialValueMediaStreamPropertyRowGenerator::Get()
 {
 	static TSharedRef<FDMMaterialValueMediaStreamPropertyRowGenerator> Generator = MakeShared<FDMMaterialValueMediaStreamPropertyRowGenerator>();
@@ -50,84 +53,197 @@ void FDMMaterialValueMediaStreamPropertyRowGenerator::AddComponentProperties(FDM
 	FDMComponentPropertyRowGeneratorParams MediaStreamParams = InParams;
 	MediaStreamParams.Object = MediaStream;
 
-	AddSourceProperty(MediaStreamParams);
-	AddPlayerProperty(MediaStreamParams);
+	const bool bHasPlayer = !!MediaStream->GetPlayer().GetInterface();
+
+	if (bHasPlayer)
+	{
+		AddControlCategory(MediaStreamParams, /* Preview */ true);
+	}
+
+	AddSourceCategory(MediaStreamParams);
+
+	if (bHasPlayer)
+	{
+		AddDetailsCategory(MediaStreamParams);
+		AddTextureCategory(MediaStreamParams, /* Preview */ true);
+		AddCacheCategory(MediaStreamParams);
+		AddPlayerProperty(MediaStreamParams);
+
+		AddControlCategory(MediaStreamParams, /* Preview */ false);
+		AddTextureCategory(MediaStreamParams, /* Preview */ false);
+	}
 }
 
-void FDMMaterialValueMediaStreamPropertyRowGenerator::AddSourceProperty(FDMComponentPropertyRowGeneratorParams& InParams)
+void FDMMaterialValueMediaStreamPropertyRowGenerator::AddControlCategory(FDMComponentPropertyRowGeneratorParams& InParams, 
+	bool bInPreview)
 {
-	FDMPropertyHandle& NewHandle = InParams.PropertyRows->Add_GetRef(
-		IDMWidgetLibrary::Get().GetPropertyHandle(InParams.CreatePropertyHandleParams(TEXT("Source")))
+	FDMComponentPropertyRowGeneratorParams* Params = &InParams;
+	UMediaStream* MediaStream = nullptr;
+	FDMComponentPropertyRowGeneratorParams OriginalParams = InParams;
+	
+	if (bInPreview)
+	{
+		MediaStream = Cast<UMediaStream>(InParams.Object);
+	}
+	else
+	{
+		const FString RelativePath = InParams.Object->GetPathName(InParams.PreviewMaterialModelBase);
+		MediaStream = UDMMaterialModelFunctionLibrary::FindSubobject<UMediaStream>(InParams.OriginalMaterialModelBase, RelativePath);
+		OriginalParams.Object = MediaStream;
+		Params = &OriginalParams;
+	}
+
+	if (!MediaStream)
+	{
+		return;
+	}
+	
+	const FName CategoryName = bInPreview
+		? "Preview Media Controls"
+		: "Source Media Controls";
+
+	FDMPropertyHandle& TrackHandle = Params->PropertyRows->Add_GetRef(
+		IDMWidgetLibrary::Get().GetPropertyHandle(Params->CreatePropertyHandleParams(TEXT("Track")))
+	);
+	TrackHandle.CategoryOverrideName = CategoryName;
+	TrackHandle.ValueWidget = UE::MediaStreamEditor::FMediaStreamWidgets::CreateTrackWidget(MediaStream);
+	TrackHandle.ValueName = bInPreview ? "PreviewTrack" : "SourceTrack";
+
+	FDMPropertyHandle& ControlHandle = Params->PropertyRows->Add_GetRef(
+		IDMWidgetLibrary::Get().GetPropertyHandle(Params->CreatePropertyHandleParams(TEXT("Controls")))
 	);
-	NewHandle.CategoryOverrideName = TEXT("Media Stream");
+	ControlHandle.CategoryOverrideName = CategoryName;
+	ControlHandle.ValueWidget = UE::MediaStreamEditor::FMediaStreamWidgets::CreateControlsWidget(MediaStream);
+	ControlHandle.ValueName = bInPreview ? "PreviewControls" : "SourceControls";
+
+	if (bInPreview)
+	{
+		FDMPropertyHandle& PlayOnOpenHandle = Params->PropertyRows->Add_GetRef(
+			IDMWidgetLibrary::Get().GetPropertyHandle(Params->CreatePropertyHandleParams(TEXT("bPlayOnOpen")))
+		);
+		PlayOnOpenHandle.CategoryOverrideName = CategoryName;
+
+		FDMPropertyHandle& LoopingHandle = Params->PropertyRows->Add_GetRef(
+			IDMWidgetLibrary::Get().GetPropertyHandle(Params->CreatePropertyHandleParams(TEXT("bLooping")))
+		);
+		LoopingHandle.CategoryOverrideName = CategoryName;
+	}
 }
 
-void FDMMaterialValueMediaStreamPropertyRowGenerator::AddPlayerProperty(FDMComponentPropertyRowGeneratorParams& InParams)
+void FDMMaterialValueMediaStreamPropertyRowGenerator::AddSourceCategory(FDMComponentPropertyRowGeneratorParams& InParams)
 {
-	UMediaStream* MediaStream = Cast<UMediaStream>(InParams.Object);
+	const FName CategoryName = "Media Source";
 
-	IMediaStreamPlayer* Player = MediaStream->GetPlayer().GetInterface();
+	UMediaStream* MediaStream = Cast<UMediaStream>(InParams.Object);
 
-	if (!Player || Player->IsReadOnly())
+	if (!MediaStream)
 	{
 		return;
 	}
 
-	const FName PreviewCategory = TEXT("Preview Controls");
-	const FName SourceCategory = TEXT("Source Controls");
-	const FName VideoPlayerCategory = TEXT("Media Stream");
+	IMediaStreamSchemeHandler::FCustomWidgets Widgets = UE::MediaStreamEditor::FMediaStreamWidgets::GenerateSourceSchemeRows(MediaStream);
 
-	const FString RelativePath = InParams.Object->GetPathName(InParams.PreviewMaterialModelBase);
-	UMediaStream* OriginalMediaStream = UDMMaterialModelFunctionLibrary::FindSubobject<UMediaStream>(InParams.OriginalMaterialModelBase, RelativePath);
+	for (const IMediaStreamSchemeHandler::FWidgetRow& WidgetRow : Widgets.CustomRows)
+	{
+		if (WidgetRow.Visibility.Get() != EVisibility::Visible)
+		{
+			continue;
+		}
+
+		FDMPropertyHandle& Handle = InParams.PropertyRows->AddDefaulted_GetRef();
+		Handle.bEnabled = WidgetRow.Enabled.Get();
+		Handle.bKeyframeable = false;
+		Handle.ValueName = *WidgetRow.Name.ToString();
+		Handle.ValueWidget = WidgetRow.Widget;
+		Handle.CategoryOverrideName = CategoryName;
+	}
+}
+
+void FDMMaterialValueMediaStreamPropertyRowGenerator::AddDetailsCategory(FDMComponentPropertyRowGeneratorParams& InParams)
+{
+	const FName CategoryName = "Media Details";
+
+	UMediaStream* MediaStream = Cast<UMediaStream>(InParams.Object);
 
-	FDMPropertyHandle& ControlHandle = InParams.PropertyRows->Add_GetRef(
-		IDMWidgetLibrary::Get().GetPropertyHandle(InParams.CreatePropertyHandleParams(TEXT("Control")))
+	FDMPropertyHandle& DetailsHandle = InParams.PropertyRows->Add_GetRef(
+		IDMWidgetLibrary::Get().GetPropertyHandle(InParams.CreatePropertyHandleParams(TEXT("Details")))
 	);
-	ControlHandle.CategoryOverrideName = PreviewCategory;
+	DetailsHandle.CategoryOverrideName = CategoryName;
+	DetailsHandle.ValueWidget = UE::MediaStreamEditor::FMediaStreamWidgets::CreateTextureDetailsWidget(MediaStream);
+	DetailsHandle.ValueName = "Details";
+}
 
-	if (OriginalMediaStream)
+void FDMMaterialValueMediaStreamPropertyRowGenerator::AddTextureCategory(FDMComponentPropertyRowGeneratorParams& InParams, 
+	bool bInPreview)
+{
+	FDMComponentPropertyRowGeneratorParams* Params = &InParams;
+	UMediaStream* MediaStream = nullptr;
+	FDMComponentPropertyRowGeneratorParams OriginalParams = InParams;
+	
+	if (bInPreview)
 	{
-		FDMComponentPropertyRowGeneratorParams SourceParams = InParams;
-		SourceParams.Object = OriginalMediaStream;
-
-		FDMPropertyHandle& SourceControlHandle = InParams.PropertyRows->Add_GetRef(
-				IDMWidgetLibrary::Get().GetPropertyHandle(SourceParams.CreatePropertyHandleParams(TEXT("Control")))
-			);
-		SourceControlHandle.CategoryOverrideName = SourceCategory;
+		MediaStream = Cast<UMediaStream>(InParams.Object);
 	}
-
+	else
 	{
-		FDMComponentPropertyRowGeneratorParams PlayerParams = InParams;
-		PlayerParams.Object = Player->_getUObject();
+		const FString RelativePath = InParams.Object->GetPathName(InParams.PreviewMaterialModelBase);
+		MediaStream = UDMMaterialModelFunctionLibrary::FindSubobject<UMediaStream>(InParams.OriginalMaterialModelBase, RelativePath);
+		OriginalParams.Object = MediaStream;
+		Params = &OriginalParams;
+	}
 
-		FDMPropertyHandle& PlaylistIndexHandle = InParams.PropertyRows->Add_GetRef(
-			IDMWidgetLibrary::Get().GetPropertyHandle(PlayerParams.CreatePropertyHandleParams(TEXT("PlaylistIndex")))
-		);
-		PlaylistIndexHandle.CategoryOverrideName = VideoPlayerCategory;
+	if (!MediaStream)
+	{
+		return;
+	}
+	
+	const FName CategoryName = bInPreview
+		? "Preview Media Texture"
+		: "Source Media Texture";
 
-		FDMPropertyHandle& PlaybackStateHandle = InParams.PropertyRows->Add_GetRef(
-			IDMWidgetLibrary::Get().GetPropertyHandle(PlayerParams.CreatePropertyHandleParams(TEXT("PlaybackState")))
-		);
-		PlaybackStateHandle.CategoryOverrideName = VideoPlayerCategory;
+	FDMPropertyHandle& MediaTextureHandle = Params->PropertyRows->Add_GetRef(
+		IDMWidgetLibrary::Get().GetPropertyHandle(Params->CreatePropertyHandleParams(TEXT("MediaTexture")))
+	);
+	MediaTextureHandle.CategoryOverrideName = CategoryName;
+	MediaTextureHandle.ValueName = bInPreview ? "MediaTexturePreview" : "MediaTextureSource";
 
-		FDMPropertyHandle& RequestedSeekFrameHandle = InParams.PropertyRows->Add_GetRef(
-			IDMWidgetLibrary::Get().GetPropertyHandle(PlayerParams.CreatePropertyHandleParams(TEXT("RequestedSeekFrame")))
+	if (bInPreview)
+	{
+		FDMPropertyHandle& TextureConfigHandle = Params->PropertyRows->Add_GetRef(
+			IDMWidgetLibrary::Get().GetPropertyHandle(Params->CreatePropertyHandleParams(TEXT("TextureConfig")))
 		);
-		RequestedSeekFrameHandle.CategoryOverrideName = VideoPlayerCategory;
+		TextureConfigHandle.CategoryOverrideName = CategoryName;
 	}
+}
+
+void FDMMaterialValueMediaStreamPropertyRowGenerator::AddCacheCategory(FDMComponentPropertyRowGeneratorParams& InParams)
+{
+	const FName CategoryName = "Media Cache";
 
-	FDMPropertyHandle& TextureHandle = InParams.PropertyRows->Add_GetRef(
-		IDMWidgetLibrary::Get().GetPropertyHandle(InParams.CreatePropertyHandleParams(TEXT("MediaTexture")))
+	FDMPropertyHandle& CacheAheadHandle = InParams.PropertyRows->Add_GetRef(
+		IDMWidgetLibrary::Get().GetPropertyHandle(InParams.CreatePropertyHandleParams(TEXT("CacheAhead")))
 	);
-	TextureHandle.CategoryOverrideName = VideoPlayerCategory;
+	CacheAheadHandle.CategoryOverrideName = CategoryName;
 
-	FDMPropertyHandle& TextureConfigHandle = InParams.PropertyRows->Add_GetRef(
-		IDMWidgetLibrary::Get().GetPropertyHandle(InParams.CreatePropertyHandleParams(TEXT("TextureConfig")))
+	FDMPropertyHandle& CacheBehindHandle = InParams.PropertyRows->Add_GetRef(
+		IDMWidgetLibrary::Get().GetPropertyHandle(InParams.CreatePropertyHandleParams(TEXT("CacheBehind")))
 	);
-	TextureConfigHandle.CategoryOverrideName = VideoPlayerCategory;
+	CacheBehindHandle.CategoryOverrideName = CategoryName;
+
+	FDMPropertyHandle& CacheBehindGameHandle = InParams.PropertyRows->Add_GetRef(
+		IDMWidgetLibrary::Get().GetPropertyHandle(InParams.CreatePropertyHandleParams(TEXT("CacheBehindGame")))
+	);
+	CacheBehindGameHandle.CategoryOverrideName = CategoryName;
+}
+
+void FDMMaterialValueMediaStreamPropertyRowGenerator::AddPlayerProperty(FDMComponentPropertyRowGeneratorParams& InParams)
+{
+	const FName CategoryName = "Media Player";
 
 	FDMPropertyHandle& PlayerConfigHandle = InParams.PropertyRows->Add_GetRef(
 		IDMWidgetLibrary::Get().GetPropertyHandle(InParams.CreatePropertyHandleParams(TEXT("PlayerConfig")))
 	);
-	PlayerConfigHandle.CategoryOverrideName = VideoPlayerCategory;
+	PlayerConfigHandle.CategoryOverrideName = CategoryName;
 }
+
+#undef LOCTEXT_NAMESPACE
diff --git a/Engine/Plugins/Experimental/DynamicMaterialMediaStreamBridge/Source/DynamicMaterialMediaStreamBridgeEditor/Private/DMMaterialValueMediaStreamPropertyRowGenerator.h b/Engine/Plugins/Experimental/DynamicMaterialMediaStreamBridge/Source/DynamicMaterialMediaStreamBridgeEditor/Private/DMMaterialValueMediaStreamPropertyRowGenerator.h
--- a/Engine/Plugins/Experimental/DynamicMaterialMediaStreamBridge/Source/DynamicMaterialMediaStreamBridgeEditor/Private/DMMaterialValueMediaStreamPropertyRowGenerator.h
+++ b/Engine/Plugins/Experimental/DynamicMaterialMediaStreamBridge/Source/DynamicMaterialMediaStreamBridgeEditor/Private/DMMaterialValueMediaStreamPropertyRowGenerator.h
@@ -17,7 +17,15 @@ public:
 	virtual void AddComponentProperties(FDMComponentPropertyRowGeneratorParams& InParams) override;
 	//~ End FDMComponentPropertyRowGenerator
 
-	void AddSourceProperty(FDMComponentPropertyRowGeneratorParams& InParams);
+	void AddControlCategory(FDMComponentPropertyRowGeneratorParams& InParams, bool bInPreview);
+
+	void AddSourceCategory(FDMComponentPropertyRowGeneratorParams& InParams);
+
+	void AddDetailsCategory(FDMComponentPropertyRowGeneratorParams& InParams);
+
+	void AddTextureCategory(FDMComponentPropertyRowGeneratorParams& InParams, bool bInPreview);
+
+	void AddCacheCategory(FDMComponentPropertyRowGeneratorParams& InParams);
 
 	void AddPlayerProperty(FDMComponentPropertyRowGeneratorParams& InParams);
 };
diff --git a/Engine/Plugins/Experimental/MediaStream/Source/MediaStream/MediaStream.Build.cs b/Engine/Plugins/Experimental/MediaStream/Source/MediaStream/MediaStream.Build.cs
--- a/Engine/Plugins/Experimental/MediaStream/Source/MediaStream/MediaStream.Build.cs
+++ b/Engine/Plugins/Experimental/MediaStream/Source/MediaStream/MediaStream.Build.cs
@@ -36,7 +36,8 @@ public class MediaStream : ModuleRules
 					"DesktopWidgets",
 					"PropertyEditor",
 					"Slate",
-					"SlateCore"
+					"SlateCore",
+					"UnrealEd"
 				}
 			);
 		}
diff --git a/Engine/Plugins/Experimental/MediaStream/Source/MediaStream/Private/MediaStream.cpp b/Engine/Plugins/Experimental/MediaStream/Source/MediaStream/Private/MediaStream.cpp
--- a/Engine/Plugins/Experimental/MediaStream/Source/MediaStream/Private/MediaStream.cpp
+++ b/Engine/Plugins/Experimental/MediaStream/Source/MediaStream/Private/MediaStream.cpp
@@ -10,6 +10,16 @@
 #include "Players/MediaStreamLocalPlayer.h"
 #include "Players/MediaStreamProxyPlayer.h"
 
+FName UMediaStream::GetSourcePropertyName()
+{
+	return GET_MEMBER_NAME_CHECKED(UMediaStream, Source);
+}
+
+FName UMediaStream::GetPlayerPropertyName()
+{
+	return GET_MEMBER_NAME_CHECKED(UMediaStream, PlayerObject);
+}
+
 bool UMediaStream::HasValidSource() const
 {
 	return UMediaStreamSourceBlueprintLibrary::IsValidMediaSource(Source);
diff --git a/Engine/Plugins/Experimental/MediaStream/Source/MediaStream/Private/MediaStreamSchemeHandlerLibrary.cpp b/Engine/Plugins/Experimental/MediaStream/Source/MediaStream/Private/MediaStreamSchemeHandlerLibrary.cpp
--- a/Engine/Plugins/Experimental/MediaStream/Source/MediaStream/Private/MediaStreamSchemeHandlerLibrary.cpp
+++ b/Engine/Plugins/Experimental/MediaStream/Source/MediaStream/Private/MediaStreamSchemeHandlerLibrary.cpp
@@ -6,6 +6,7 @@
 
 #if WITH_EDITOR
 #include "MediaStreamSchemeHandlerManager.h"
+#include "MediaStream.h"
 #include "MediaStreamSource.h"
 #include "PropertyHandle.h"
 #endif
@@ -62,6 +63,21 @@ FName FMediaStreamSchemeHandlerLibrary::GetScheme(TWeakPtr<IPropertyHandle> InPr
 	return NAME_None;
 }
 
+FName FMediaStreamSchemeHandlerLibrary::GetScheme(UMediaStream* InMediaStream)
+{
+	if (InMediaStream)
+	{
+		return InMediaStream->GetSource().Scheme;
+	}
+
+	return NAME_None;
+}
+
+FName FMediaStreamSchemeHandlerLibrary::GetScheme(TWeakObjectPtr<UMediaStream> InMediaStreamWeak)
+{
+	return GetScheme(InMediaStreamWeak.Get());
+}
+
 FString FMediaStreamSchemeHandlerLibrary::GetPath(TWeakPtr<IPropertyHandle> InPropertyHandleWeak)
 {
 	if (TSharedPtr<IPropertyHandle> PropertyHandle = InPropertyHandleWeak.Pin())
@@ -72,7 +88,22 @@ FString FMediaStreamSchemeHandlerLibrary::GetPath(TWeakPtr<IPropertyHandle> InPr
 		}
 	}
 
-	return "";
+	return FString();
+}
+
+FString FMediaStreamSchemeHandlerLibrary::GetPath(UMediaStream* InMediaStream)
+{
+	if (InMediaStream)
+	{
+		return InMediaStream->GetSource().Path;
+	}
+
+	return FString();
+}
+
+FString FMediaStreamSchemeHandlerLibrary::GetPath(TWeakObjectPtr<UMediaStream> InMediaStreamWeak)
+{
+	return GetPath(InMediaStreamWeak.Get());
 }
 
 void FMediaStreamSchemeHandlerLibrary::SetSource(TSharedRef<IPropertyHandle> InPropertyHandle, const FName& InScheme, const FString& InPath)
@@ -101,4 +132,26 @@ void FMediaStreamSchemeHandlerLibrary::SetSource(TSharedRef<IPropertyHandle> InP
 	Outers[0]->PostEditChangeProperty(Event);
 }
 
+void FMediaStreamSchemeHandlerLibrary::SetSource(UMediaStream* InMediaStream, const FName& InScheme, const FString& InPath)
+{
+	if (!IsValid(InMediaStream))
+	{
+		return;
+	}
+
+	FProperty* Property = InMediaStream->GetClass()->FindPropertyByName( UMediaStream::GetSourcePropertyName() );
+
+	InMediaStream->PreEditChange(Property);
+
+	InMediaStream->SetSource(FMediaStreamSchemeHandlerManager::Get().CreateSource(InMediaStream, InScheme, InPath));
+
+	FPropertyChangedEvent Event(Property, EPropertyChangeType::Interactive, {InMediaStream});
+	InMediaStream->PostEditChangeProperty(Event);
+}
+
+void FMediaStreamSchemeHandlerLibrary::SetSource(TWeakObjectPtr<UMediaStream> InMediaStreamWeak, const FName& InScheme, const FString& InPath)
+{
+	SetSource(InMediaStreamWeak.Get(), InScheme, InPath);
+}
+
 #endif
diff --git a/Engine/Plugins/Experimental/MediaStream/Source/MediaStream/Private/Players/MediaStreamLocalPlayer.h b/Engine/Plugins/Experimental/MediaStream/Source/MediaStream/Private/Players/MediaStreamLocalPlayer.h
--- a/Engine/Plugins/Experimental/MediaStream/Source/MediaStream/Private/Players/MediaStreamLocalPlayer.h
+++ b/Engine/Plugins/Experimental/MediaStream/Source/MediaStream/Private/Players/MediaStreamLocalPlayer.h
@@ -118,7 +118,7 @@ protected:
 
 	UPROPERTY(EditAnywhere, BlueprintReadWrite, Setter = BP_SetPlaybackState, Getter, BlueprintSetter = BP_SetPlaybackState,
 		Category = "Media Stream|Player", Interp, meta = (AllowPrivateAccess = true))
-	EMediaStreamPlaybackState PlaybackState = EMediaStreamPlaybackState::Play;
+	EMediaStreamPlaybackState PlaybackState = EMediaStreamPlaybackState::Pause;
 
 	/** The index playing in the playlist. -1 to ignore. */
 	UPROPERTY(EditAnywhere, BlueprintReadWrite, Setter = BP_SetPlaylistIndex, Getter, BlueprintSetter = BP_SetPlaylistIndex,
diff --git a/Engine/Plugins/Experimental/MediaStream/Source/MediaStream/Private/SchemeHandlers/MediaStreamAssetSchemeHandler.cpp b/Engine/Plugins/Experimental/MediaStream/Source/MediaStream/Private/SchemeHandlers/MediaStreamAssetSchemeHandler.cpp
--- a/Engine/Plugins/Experimental/MediaStream/Source/MediaStream/Private/SchemeHandlers/MediaStreamAssetSchemeHandler.cpp
+++ b/Engine/Plugins/Experimental/MediaStream/Source/MediaStream/Private/SchemeHandlers/MediaStreamAssetSchemeHandler.cpp
@@ -11,10 +11,8 @@
 #include "MediaStreamSourceBlueprintLibrary.h"
 
 #if WITH_EDITOR
-#include "DetailLayoutBuilder.h"
-#include "DetailWidgetRow.h"
-#include "IDetailGroup.h"
 #include "PropertyCustomizationHelpers.h"
+#include "ThumbnailRendering/ThumbnailManager.h"
 #include "Widgets/Text/STextBlock.h"
 #endif
 
@@ -102,55 +100,36 @@ UObject* FMediaStreamAssetSchemeHandler::ResolveAsset(const FString& InPath)
 }
 
 #if WITH_EDITOR
-void FMediaStreamAssetSchemeHandler::CreatePropertyCustomization(TSharedRef<IPropertyHandle> InPropertyHandle, IDetailGroup& InDetailsGroup, IPropertyTypeCustomizationUtils& InCustomizationUtils)
+void FMediaStreamAssetSchemeHandler::CreatePropertyCustomization(UMediaStream* InMediaStream, IMediaStreamSchemeHandler::FCustomWidgets& InOutCustomWidgets)
 {
-	AddAssetSelector(InPropertyHandle, InDetailsGroup, InCustomizationUtils);
+	AddAssetSelector(InMediaStream, InOutCustomWidgets);
 }
 
-void FMediaStreamAssetSchemeHandler::AddAssetSelector(TSharedRef<IPropertyHandle> InPropertyHandle, IDetailGroup& InSourceGroup, IPropertyTypeCustomizationUtils& InCustomizationUtils)
+void FMediaStreamAssetSchemeHandler::AddAssetSelector(UMediaStream* InMediaStream, IMediaStreamSchemeHandler::FCustomWidgets& InOutCustomWidgets)
 {
-	TWeakPtr<IPropertyHandle> PropertyHandleWeak = InPropertyHandle;
-
-	FMediaStreamSource* StreamSource = FMediaStreamSchemeHandlerLibrary::GetStreamSourcePtr(InPropertyHandle);
-
-	if (!StreamSource)
+	if (!IsValid(InMediaStream))
 	{
 		return;
 	}
 
-	TSharedPtr<IPropertyHandle> AssetProperty = InPropertyHandle->GetChildHandle(GET_MEMBER_NAME_CHECKED(FMediaStreamSource, Object), /* Recurse */ false);
-
-	if (!AssetProperty.IsValid())
-	{
-		return;
-	}
+	TWeakObjectPtr<UMediaStream> MediaStreamWeak = InMediaStream;
 
-	FDetailWidgetRow& AssetSelectRow = InSourceGroup.AddWidgetRow();
-
-	AssetSelectRow.VisibilityAttr = TAttribute<EVisibility>::CreateSP(this, &FMediaStreamAssetSchemeHandler::GetAssetSelectorVisibility, PropertyHandleWeak);
-
-	AssetSelectRow.NameContent()
-	[
-		SNew(STextBlock)
-		.Font(IDetailLayoutBuilder::GetDetailFont())
-		.Text(LOCTEXT("Asset", "Asset"))
-	];
-
-	AssetSelectRow.ValueContent()
-	[
+	InOutCustomWidgets.CustomRows.Add({
+		LOCTEXT("Asset", "Asset"),
 		SNew(SObjectPropertyEntryBox)
-		.ObjectPath_Static(&FMediaStreamSchemeHandlerLibrary::GetPath, PropertyHandleWeak)
-		.OnObjectChanged(this, &FMediaStreamAssetSchemeHandler::OnAssetSelected, PropertyHandleWeak)
-		.ThumbnailPool(InCustomizationUtils.GetThumbnailPool())
+		.ObjectPath_Static(&FMediaStreamSchemeHandlerLibrary::GetPath, MediaStreamWeak)
+		.OnObjectChanged(this, &FMediaStreamAssetSchemeHandler::OnAssetSelected, MediaStreamWeak)
+		.ThumbnailPool(UThumbnailManager::Get().GetSharedThumbnailPool())
 		.DisplayThumbnail(true)
-		.AllowedClass(UMediaSource::StaticClass())
-		.PropertyHandle(InPropertyHandle)
-	];
+		.AllowedClass(UMediaSource::StaticClass()),
+		/* Enabled */ true,
+		TAttribute<EVisibility>::CreateSP(this, &FMediaStreamAssetSchemeHandler::GetAssetSelectorVisibility, MediaStreamWeak)
+	});
 }
 
-EVisibility FMediaStreamAssetSchemeHandler::GetAssetSelectorVisibility(TWeakPtr<IPropertyHandle> InPropertyHandleWeak) const
+EVisibility FMediaStreamAssetSchemeHandler::GetAssetSelectorVisibility(TWeakObjectPtr<UMediaStream> InMediaStreamWeak) const
 {
-	if (FMediaStreamSchemeHandlerLibrary::GetScheme(InPropertyHandleWeak) == Scheme)
+	if (FMediaStreamSchemeHandlerLibrary::GetScheme(InMediaStreamWeak) == Scheme)
 	{
 		return EVisibility::Visible;
 	}
@@ -158,12 +137,9 @@ EVisibility FMediaStreamAssetSchemeHandler::GetAssetSelectorVisibility(TWeakPtr<
 	return EVisibility::Collapsed;
 }
 
-void FMediaStreamAssetSchemeHandler::OnAssetSelected(const FAssetData& InAssetData, TWeakPtr<IPropertyHandle> InPropertyHandleWeak)
+void FMediaStreamAssetSchemeHandler::OnAssetSelected(const FAssetData& InAssetData, TWeakObjectPtr<UMediaStream> InMediaStreamWeak)
 {
-	if (TSharedPtr<IPropertyHandle> PropertyHandle = InPropertyHandleWeak.Pin())
-	{
-		FMediaStreamSchemeHandlerLibrary::SetSource(PropertyHandle.ToSharedRef(), Scheme, InAssetData.GetSoftObjectPath().ToString());
-	}
+	FMediaStreamSchemeHandlerLibrary::SetSource(InMediaStreamWeak, Scheme, InAssetData.GetSoftObjectPath().ToString());
 }
 #endif
 
diff --git a/Engine/Plugins/Experimental/MediaStream/Source/MediaStream/Private/SchemeHandlers/MediaStreamAssetSchemeHandler.h b/Engine/Plugins/Experimental/MediaStream/Source/MediaStream/Private/SchemeHandlers/MediaStreamAssetSchemeHandler.h
--- a/Engine/Plugins/Experimental/MediaStream/Source/MediaStream/Private/SchemeHandlers/MediaStreamAssetSchemeHandler.h
+++ b/Engine/Plugins/Experimental/MediaStream/Source/MediaStream/Private/SchemeHandlers/MediaStreamAssetSchemeHandler.h
@@ -21,7 +21,7 @@ public:
 	virtual UMediaPlayer* CreateOrUpdatePlayer(const FMediaStreamSchemeHandlerCreatePlayerParams& InParams) override;
 
 #if WITH_EDITOR
-	virtual void CreatePropertyCustomization(TSharedRef<IPropertyHandle> InPropertyHandle, IDetailGroup& InDetailsGroup, IPropertyTypeCustomizationUtils& InCustomizationUtils) override;
+	virtual void CreatePropertyCustomization(UMediaStream* InMediaStream, IMediaStreamSchemeHandler::FCustomWidgets& InOutCustomWidgets) override;
 #endif
 	//~ End IMediaStreamSchemeHandler
 
@@ -29,10 +29,10 @@ protected:
 	UObject* ResolveAsset(const FString& InPath);
 
 #if WITH_EDITOR
-	void AddAssetSelector(TSharedRef<IPropertyHandle> InPropertyHandle, IDetailGroup& InSourceGroup, IPropertyTypeCustomizationUtils& InCustomizationUtils);
+	void AddAssetSelector(UMediaStream* InMediaStream, IMediaStreamSchemeHandler::FCustomWidgets& InOutCustomWidgets);
 
-	EVisibility GetAssetSelectorVisibility(TWeakPtr<IPropertyHandle> InPropertyHandleWeak) const;
+	EVisibility GetAssetSelectorVisibility(TWeakObjectPtr<UMediaStream> InMediaStreamWeak) const;
 
-	void OnAssetSelected(const FAssetData& InAssetData, TWeakPtr<IPropertyHandle> InPropertyHandleWeak);
+	void OnAssetSelected(const FAssetData& InAssetData, TWeakObjectPtr<UMediaStream> InMediaStreamWeak);
 #endif
 };
diff --git a/Engine/Plugins/Experimental/MediaStream/Source/MediaStream/Private/SchemeHandlers/MediaStreamFileSchemeHandler.cpp b/Engine/Plugins/Experimental/MediaStream/Source/MediaStream/Private/SchemeHandlers/MediaStreamFileSchemeHandler.cpp
--- a/Engine/Plugins/Experimental/MediaStream/Source/MediaStream/Private/SchemeHandlers/MediaStreamFileSchemeHandler.cpp
+++ b/Engine/Plugins/Experimental/MediaStream/Source/MediaStream/Private/SchemeHandlers/MediaStreamFileSchemeHandler.cpp
@@ -12,9 +12,6 @@
 #include "Misc/Paths.h"
 
 #if WITH_EDITOR
-#include "DetailLayoutBuilder.h"
-#include "DetailWidgetRow.h"
-#include "IDetailGroup.h"
 #include "Widgets/Input/SFilePathPicker.h"
 #include "Widgets/Text/STextBlock.h"
 #endif
@@ -105,7 +102,8 @@ FString FMediaStreamFileSchemeHandler::ResolveFilePath(const FString& InPath)
 	}
 
 	UE_LOG(LogMediaStream, Error, TEXT("Invalid File Path in FMediaStreamFileSchemeHandler::ResolveFilePath [%s]"), *InPath);
-	return "";
+
+	return FString();
 }
 
 UMediaSource* FMediaStreamFileSchemeHandler::CreateMediaSource(UObject* InOuter, const FString& InPath)
@@ -129,42 +127,38 @@ UMediaSource* FMediaStreamFileSchemeHandler::CreateMediaSource(UObject* InOuter,
 }
 
 #if WITH_EDITOR
-void FMediaStreamFileSchemeHandler::CreatePropertyCustomization(TSharedRef<IPropertyHandle> InPropertyHandle, IDetailGroup& InDetailsGroup, IPropertyTypeCustomizationUtils& InCustomizationUtils)
+void FMediaStreamFileSchemeHandler::CreatePropertyCustomization(UMediaStream* InMediaStream, IMediaStreamSchemeHandler::FCustomWidgets& InOutCustomWidgets)
 {
-	AddFileSelector(InPropertyHandle, InDetailsGroup, InCustomizationUtils);
+	AddFileSelector(InMediaStream, InOutCustomWidgets);
 }
 
-void FMediaStreamFileSchemeHandler::AddFileSelector(TSharedRef<IPropertyHandle> InPropertyHandle, IDetailGroup& InSourceGroup, IPropertyTypeCustomizationUtils& InCustomizationUtils)
+void FMediaStreamFileSchemeHandler::AddFileSelector(UMediaStream* InMediaStream, IMediaStreamSchemeHandler::FCustomWidgets& InOutCustomWidgets)
 {
-	TWeakPtr<IPropertyHandle> PropertyHandleWeak = InPropertyHandle;
-
-	FDetailWidgetRow& FilePathRow = InSourceGroup.AddWidgetRow();
-
-	FilePathRow.VisibilityAttr = TAttribute<EVisibility>::CreateSP(this, &FMediaStreamFileSchemeHandler::GetFileSelectorVisibility, PropertyHandleWeak);
+	if (!IsValid(InMediaStream))
+	{
+		return;
+	}
 
-	FilePathRow.NameContent()
-	[
-		SNew(STextBlock)
-		.Font(IDetailLayoutBuilder::GetDetailFont())
-		.Text(LOCTEXT("FilePath", "File Path"))
-	];
+	TWeakObjectPtr<UMediaStream> MediaStreamWeak = InMediaStream;
 
-	FilePathRow.ValueContent()
-	[
+	InOutCustomWidgets.CustomRows.Add({
+		LOCTEXT("FilePath", "File Path"),
 		SNew(SFilePathPicker)
 		.BrowseButtonImage(FAppStyle::GetBrush("PropertyWindow.Button_Ellipsis"))
 		.BrowseButtonStyle(FAppStyle::Get(), "HoverHintOnly")
 		.BrowseButtonToolTip(LOCTEXT("FileButtonToolTipText", "Choose a file from this computer"))
 		.BrowseTitle(LOCTEXT("PropertyEditorTitle", "File picker..."))
-		.FilePath_Static(&FMediaStreamSchemeHandlerLibrary::GetPath, PropertyHandleWeak)
+		.FilePath_Static(&FMediaStreamSchemeHandlerLibrary::GetPath, MediaStreamWeak)
 		.FileTypeFilter(TEXT("All files (*.*)|*.*"))
-		.OnPathPicked(this, &FMediaStreamFileSchemeHandler::OnFilePicked, PropertyHandleWeak)
-	];
+		.OnPathPicked(this, &FMediaStreamFileSchemeHandler::OnFilePicked, MediaStreamWeak),
+		/* Enabled */ true,
+		TAttribute<EVisibility>::CreateSP(this, &FMediaStreamFileSchemeHandler::GetFileSelectorVisibility, MediaStreamWeak)
+	});
 }
 
-EVisibility FMediaStreamFileSchemeHandler::GetFileSelectorVisibility(TWeakPtr<IPropertyHandle> InPropertyHandleWeak) const
+EVisibility FMediaStreamFileSchemeHandler::GetFileSelectorVisibility(TWeakObjectPtr<UMediaStream> InMediaStreamWeak) const
 {
-	if (FMediaStreamSchemeHandlerLibrary::GetScheme(InPropertyHandleWeak) == Scheme)
+	if (FMediaStreamSchemeHandlerLibrary::GetScheme(InMediaStreamWeak) == Scheme)
 	{
 		return EVisibility::Visible;
 	}
@@ -172,12 +166,9 @@ EVisibility FMediaStreamFileSchemeHandler::GetFileSelectorVisibility(TWeakPtr<IP
 	return EVisibility::Collapsed;
 }
 
-void FMediaStreamFileSchemeHandler::OnFilePicked(const FString& InFilePath, TWeakPtr<IPropertyHandle> InPropertyHandleWeak)
+void FMediaStreamFileSchemeHandler::OnFilePicked(const FString& InFilePath, TWeakObjectPtr<UMediaStream> InMediaStreamWeak)
 {
-	if (TSharedPtr<IPropertyHandle> PropertyHandle = InPropertyHandleWeak.Pin())
-	{
-		FMediaStreamSchemeHandlerLibrary::SetSource(PropertyHandle.ToSharedRef(), Scheme, InFilePath);
-	}
+	FMediaStreamSchemeHandlerLibrary::SetSource(InMediaStreamWeak, Scheme, InFilePath);
 }
 
 #endif
diff --git a/Engine/Plugins/Experimental/MediaStream/Source/MediaStream/Private/SchemeHandlers/MediaStreamFileSchemeHandler.h b/Engine/Plugins/Experimental/MediaStream/Source/MediaStream/Private/SchemeHandlers/MediaStreamFileSchemeHandler.h
--- a/Engine/Plugins/Experimental/MediaStream/Source/MediaStream/Private/SchemeHandlers/MediaStreamFileSchemeHandler.h
+++ b/Engine/Plugins/Experimental/MediaStream/Source/MediaStream/Private/SchemeHandlers/MediaStreamFileSchemeHandler.h
@@ -17,7 +17,7 @@ public:
 	virtual UMediaPlayer* CreateOrUpdatePlayer(const FMediaStreamSchemeHandlerCreatePlayerParams& InParams) override;
 
 #if WITH_EDITOR
-	virtual void CreatePropertyCustomization(TSharedRef<IPropertyHandle> InPropertyHandle, IDetailGroup& InDetailsGroup, IPropertyTypeCustomizationUtils& InCustomizationUtils) override;
+	virtual void CreatePropertyCustomization(UMediaStream* InMediaStream, IMediaStreamSchemeHandler::FCustomWidgets& InOutCustomWidgets) override;
 #endif
 	//~ End IMediaStreamSchemeHandler
 
@@ -27,10 +27,10 @@ protected:
 	UMediaSource* CreateMediaSource(UObject* InOuter, const FString& InPath);
 
 #if WITH_EDITOR
-	void AddFileSelector(TSharedRef<IPropertyHandle> InPropertyHandle, IDetailGroup& InSourceGroup, IPropertyTypeCustomizationUtils& InCustomizationUtils);
+	void AddFileSelector(UMediaStream* InMediaStream, IMediaStreamSchemeHandler::FCustomWidgets& InOutCustomWidgets);
 
-	EVisibility GetFileSelectorVisibility(TWeakPtr<IPropertyHandle> InPropertyHandleWeak) const;
+	EVisibility GetFileSelectorVisibility(TWeakObjectPtr<UMediaStream> InMediaStreamWeak) const;
 
-	void OnFilePicked(const FString& InFilePath, TWeakPtr<IPropertyHandle> InPropertyHandleWeak);
+	void OnFilePicked(const FString& InFilePath, TWeakObjectPtr<UMediaStream> InMediaStreamWeak);
 #endif
 };
diff --git a/Engine/Plugins/Experimental/MediaStream/Source/MediaStream/Private/SchemeHandlers/MediaStreamManagedSchemeHandler.cpp b/Engine/Plugins/Experimental/MediaStream/Source/MediaStream/Private/SchemeHandlers/MediaStreamManagedSchemeHandler.cpp
--- a/Engine/Plugins/Experimental/MediaStream/Source/MediaStream/Private/SchemeHandlers/MediaStreamManagedSchemeHandler.cpp
+++ b/Engine/Plugins/Experimental/MediaStream/Source/MediaStream/Private/SchemeHandlers/MediaStreamManagedSchemeHandler.cpp
@@ -39,7 +39,7 @@ UMediaPlayer* FMediaStreamManagedSchemeHandler::CreateOrUpdatePlayer(const FMedi
 }
 
 #if WITH_EDITOR
-void FMediaStreamManagedSchemeHandler::CreatePropertyCustomization(TSharedRef<IPropertyHandle> InPropertyHandle, IDetailGroup& InDetailsGroup, IPropertyTypeCustomizationUtils& InCustomizationUtils)
+void FMediaStreamManagedSchemeHandler::CreatePropertyCustomization(UMediaStream* InMediaStream, IMediaStreamSchemeHandler::FCustomWidgets& InOutCustomWidgets)
 {
 	// TODO Implement proper customisation
 }
diff --git a/Engine/Plugins/Experimental/MediaStream/Source/MediaStream/Private/SchemeHandlers/MediaStreamManagedSchemeHandler.h b/Engine/Plugins/Experimental/MediaStream/Source/MediaStream/Private/SchemeHandlers/MediaStreamManagedSchemeHandler.h
--- a/Engine/Plugins/Experimental/MediaStream/Source/MediaStream/Private/SchemeHandlers/MediaStreamManagedSchemeHandler.h
+++ b/Engine/Plugins/Experimental/MediaStream/Source/MediaStream/Private/SchemeHandlers/MediaStreamManagedSchemeHandler.h
@@ -14,7 +14,7 @@ public:
 	virtual UMediaPlayer* CreateOrUpdatePlayer(const FMediaStreamSchemeHandlerCreatePlayerParams& InParams) override;
 
 #if WITH_EDITOR
-	virtual void CreatePropertyCustomization(TSharedRef<IPropertyHandle> InPropertyHandle, IDetailGroup& InDetailsGroup, IPropertyTypeCustomizationUtils& InCustomizationUtils) override;
+	virtual void CreatePropertyCustomization(UMediaStream* InMediaStream, IMediaStreamSchemeHandler::FCustomWidgets& InOutCustomWidgets) override;
 #endif
 	//~ End IMediaStreamSchemeHandler
 };
diff --git a/Engine/Plugins/Experimental/MediaStream/Source/MediaStream/Private/SchemeHandlers/MediaStreamSubobjectSchemeHandler.cpp b/Engine/Plugins/Experimental/MediaStream/Source/MediaStream/Private/SchemeHandlers/MediaStreamSubobjectSchemeHandler.cpp
--- a/Engine/Plugins/Experimental/MediaStream/Source/MediaStream/Private/SchemeHandlers/MediaStreamSubobjectSchemeHandler.cpp
+++ b/Engine/Plugins/Experimental/MediaStream/Source/MediaStream/Private/SchemeHandlers/MediaStreamSubobjectSchemeHandler.cpp
@@ -61,7 +61,7 @@ UObject* FMediaStreamSubobjectSchemeHandler::ResolveSubobjectPath(UObject* InOut
 }
 
 #if WITH_EDITOR
-void FMediaStreamSubobjectSchemeHandler::CreatePropertyCustomization(TSharedRef<IPropertyHandle> InPropertyHandle, IDetailGroup& InDetailsGroup, IPropertyTypeCustomizationUtils& InCustomizationUtils)
+void FMediaStreamSubobjectSchemeHandler::CreatePropertyCustomization(UMediaStream* InMediaStream, IMediaStreamSchemeHandler::FCustomWidgets& InOutCustomWidgets)
 {
 	// TODO Implement proper customisation
 }
diff --git a/Engine/Plugins/Experimental/MediaStream/Source/MediaStream/Private/SchemeHandlers/MediaStreamSubobjectSchemeHandler.h b/Engine/Plugins/Experimental/MediaStream/Source/MediaStream/Private/SchemeHandlers/MediaStreamSubobjectSchemeHandler.h
--- a/Engine/Plugins/Experimental/MediaStream/Source/MediaStream/Private/SchemeHandlers/MediaStreamSubobjectSchemeHandler.h
+++ b/Engine/Plugins/Experimental/MediaStream/Source/MediaStream/Private/SchemeHandlers/MediaStreamSubobjectSchemeHandler.h
@@ -16,7 +16,7 @@ public:
 	virtual UMediaPlayer* CreateOrUpdatePlayer(const FMediaStreamSchemeHandlerCreatePlayerParams& InParams) override;
 
 #if WITH_EDITOR
-	virtual void CreatePropertyCustomization(TSharedRef<IPropertyHandle> InPropertyHandle, IDetailGroup& InDetailsGroup, IPropertyTypeCustomizationUtils& InCustomizationUtils) override;
+	virtual void CreatePropertyCustomization(UMediaStream* InMediaStream, IMediaStreamSchemeHandler::FCustomWidgets& InOutCustomWidgets) override;
 #endif
 	//~ End IMediaStreamSchemeHandler
 
diff --git a/Engine/Plugins/Experimental/MediaStream/Source/MediaStream/Public/IMediaStreamSchemeHandler.h b/Engine/Plugins/Experimental/MediaStream/Source/MediaStream/Public/IMediaStreamSchemeHandler.h
--- a/Engine/Plugins/Experimental/MediaStream/Source/MediaStream/Public/IMediaStreamSchemeHandler.h
+++ b/Engine/Plugins/Experimental/MediaStream/Source/MediaStream/Public/IMediaStreamSchemeHandler.h
@@ -6,19 +6,20 @@
 
 #include "Containers/ContainersFwd.h"
 
+#if WITH_EDITOR
+#include "Layout/Visibility.h"
+#include "Misc/Attribute.h"
+#endif
+
 #include "IMediaStreamSchemeHandler.generated.h"
 
 class IPropertyHandle;
+class SWidget;
 class UMediaPlayer;
 class UMediaStream;
 struct FMediaStreamSource;
 struct FMediaStreamObjectHandlerCreatePlayerParams;
 
-#if WITH_EDITOR
-class IDetailGroup;
-class IPropertyTypeCustomizationUtils;
-#endif
-
 USTRUCT(BlueprintType)
 struct FMediaStreamSchemeHandlerCreatePlayerParams
 {
@@ -42,7 +43,7 @@ struct FMediaStreamSchemeHandlerCreatePlayerParams
 	 * an existing player is not updated.
 	 */
 	UPROPERTY(BlueprintReadWrite, Category = "Media Stream")
-	bool bCanOpenSource;
+	bool bCanOpenSource = true;
 
 	FMediaStreamObjectHandlerCreatePlayerParams operator<<(UObject* InMediaSource) const;
 };
@@ -72,7 +73,20 @@ public:
 	virtual UMediaPlayer* CreateOrUpdatePlayer(const FMediaStreamSchemeHandlerCreatePlayerParams& InParams) = 0;
 
 #if WITH_EDITOR
-	virtual void CreatePropertyCustomization(TSharedRef<IPropertyHandle> InPropertyHandle, IDetailGroup& InDetailsGroup, IPropertyTypeCustomizationUtils& InCustomizationUtils) = 0;
+	struct FWidgetRow
+	{
+		FText Name;
+		TSharedRef<SWidget> Widget;
+		TAttribute<bool> Enabled = true;
+		TAttribute<EVisibility> Visibility = EVisibility::Visible;
+	};
+
+	struct FCustomWidgets
+	{
+		TArray<FWidgetRow> CustomRows;
+	};
+
+	virtual void CreatePropertyCustomization(UMediaStream* InMediaStream, IMediaStreamSchemeHandler::FCustomWidgets& InOutCustomWidgets) = 0;
 #endif
 };
 
@@ -83,9 +97,15 @@ public:
 	MEDIASTREAM_API static FMediaStreamSource* GetStreamSourcePtr(TWeakPtr<IPropertyHandle> InPropertyHandleWeak);
 
 	MEDIASTREAM_API static FName GetScheme(TWeakPtr<IPropertyHandle> InPropertyHandleWeak);
+	MEDIASTREAM_API static FName GetScheme(UMediaStream* InMediaStream);
+	MEDIASTREAM_API static FName GetScheme(TWeakObjectPtr<UMediaStream> InMediaStreamWeak);
 
 	MEDIASTREAM_API static FString GetPath(TWeakPtr<IPropertyHandle> InPropertyHandleWeak);
+	MEDIASTREAM_API static FString GetPath(UMediaStream* InMediaStream);
+	MEDIASTREAM_API static FString GetPath(TWeakObjectPtr<UMediaStream> InMediaStreamWeak);
 
 	MEDIASTREAM_API static void SetSource(TSharedRef<IPropertyHandle> InPropertyHandle, const FName& InScheme, const FString& InPath);
+	MEDIASTREAM_API static void SetSource(UMediaStream* InMediaStream, const FName& InScheme, const FString& InPath);
+	MEDIASTREAM_API static void SetSource(TWeakObjectPtr<UMediaStream> InMediaStreamWeak, const FName& InScheme, const FString& InPath);
 };
 #endif
diff --git a/Engine/Plugins/Experimental/MediaStream/Source/MediaStream/Public/MediaStream.h b/Engine/Plugins/Experimental/MediaStream/Source/MediaStream/Public/MediaStream.h
--- a/Engine/Plugins/Experimental/MediaStream/Source/MediaStream/Public/MediaStream.h
+++ b/Engine/Plugins/Experimental/MediaStream/Source/MediaStream/Public/MediaStream.h
@@ -15,12 +15,16 @@ class IMediaStreamPlayer;
 /**
  * Media Stream. Provides an agnostic interface between controllers and players.
  */
-UCLASS(MinimalAPI, BlueprintType, Blueprintable, ClassGroup = "Media Stream", DefaultToInstanced)
+UCLASS(MinimalAPI, BlueprintType, Blueprintable, ClassGroup = "Media Stream", DefaultToInstanced, 
+	PrioritizeCategories = (MediaStream, MediaControls, MediaSource, MediaDetails, MediaTexture, MediaCache, MediaPlayer))
 class UMediaStream : public UObject, public IMediaPlayerProxyInterface
 {
 	GENERATED_BODY()
 
 public:
+	MEDIASTREAM_API static FName GetSourcePropertyName();
+	MEDIASTREAM_API static FName GetPlayerPropertyName();
+
 	virtual ~UMediaStream() override = default;
 
 	/**
@@ -95,7 +99,7 @@ public:
 	//~ End IMediaPlayerProxyInterface
 
 protected:
-	UPROPERTY(EditAnywhere, BlueprintReadWrite, Category = "Media Stream")
+	UPROPERTY(EditAnywhere, BlueprintReadWrite, Category = "Media Source")
 	FMediaStreamSource Source;
 
 	/**
diff --git a/Engine/Plugins/Experimental/MediaStream/Source/MediaStream/Public/MediaStreamActor.h b/Engine/Plugins/Experimental/MediaStream/Source/MediaStream/Public/MediaStreamActor.h
--- a/Engine/Plugins/Experimental/MediaStream/Source/MediaStream/Public/MediaStreamActor.h
+++ b/Engine/Plugins/Experimental/MediaStream/Source/MediaStream/Public/MediaStreamActor.h
@@ -8,7 +8,8 @@
 
 class UMediaStreamComponent;
 
-UCLASS(BlueprintType, MinimalAPI, ComponentWrapperClass)
+UCLASS(BlueprintType, MinimalAPI, ComponentWrapperClass,
+	PrioritizeCategories = (MediaStream, MediaControls, MediaSource, MediaDetails, MediaTexture, MediaCache, MediaPlayer))
 class AMediaStreamActor : public AActor
 {
 	GENERATED_BODY()
diff --git a/Engine/Plugins/Experimental/MediaStream/Source/MediaStream/Public/MediaStreamComponent.h b/Engine/Plugins/Experimental/MediaStream/Source/MediaStream/Public/MediaStreamComponent.h
--- a/Engine/Plugins/Experimental/MediaStream/Source/MediaStream/Public/MediaStreamComponent.h
+++ b/Engine/Plugins/Experimental/MediaStream/Source/MediaStream/Public/MediaStreamComponent.h
@@ -9,7 +9,8 @@
 class IMediaStreamPlayer;
 class UMediaStream;
 
-UCLASS(BlueprintType, MinimalAPI)
+UCLASS(BlueprintType, MinimalAPI,
+	PrioritizeCategories = (MediaStream, MediaControls, MediaSource, MediaDetails, MediaTexture, MediaCache, MediaPlayer))
 class UMediaStreamComponent : public UStaticMeshComponent
 {
 	GENERATED_BODY()
diff --git a/Engine/Plugins/Experimental/MediaStream/Source/MediaStream/Public/MediaStreamPlayerConfig.h b/Engine/Plugins/Experimental/MediaStream/Source/MediaStream/Public/MediaStreamPlayerConfig.h
--- a/Engine/Plugins/Experimental/MediaStream/Source/MediaStream/Public/MediaStreamPlayerConfig.h
+++ b/Engine/Plugins/Experimental/MediaStream/Source/MediaStream/Public/MediaStreamPlayerConfig.h
@@ -38,15 +38,15 @@ struct FMediaStreamPlayerConfig
 	UPROPERTY(EditInstanceOnly, BlueprintReadWrite, Category = "Media Stream|Player")
 	bool bShuffle = false;
 
-	UPROPERTY(EditInstanceOnly, BlueprintReadWrite, Category = "Media Stream|Player", AdvancedDisplay, 
+	UPROPERTY(EditInstanceOnly, BlueprintReadWrite, Category = "Media Stream|Player", 
 		meta = (Units = "s", UIMin = 0.1, ClampMin = 0.01, UIMax = 5, ClampMax = 60))
 	float CacheAhead = 0.1f;
 
-	UPROPERTY(EditInstanceOnly, BlueprintReadWrite, Category = "Media Stream|Player", AdvancedDisplay,
+	UPROPERTY(EditInstanceOnly, BlueprintReadWrite, Category = "Media Stream|Player",
 		meta = (Units = "s", UIMin = 0.1, ClampMin = 0.01, UIMax = 5, ClampMax = 60))
 	float CacheBehind = 3.0f;
 
-	UPROPERTY(EditInstanceOnly, BlueprintReadWrite, Category = "Media Stream|Player", AdvancedDisplay,
+	UPROPERTY(EditInstanceOnly, BlueprintReadWrite, Category = "Media Stream|Player",
 		meta = (Units = "s", UIMin = 0.1, ClampMin = 0.01, UIMax = 5, ClampMax = 60))
 	float CacheBehindGame = 0.1f;
 
diff --git a/Engine/Plugins/Experimental/MediaStream/Source/MediaStream/Public/MediaStreamSource.h b/Engine/Plugins/Experimental/MediaStream/Source/MediaStream/Public/MediaStreamSource.h
--- a/Engine/Plugins/Experimental/MediaStream/Source/MediaStream/Public/MediaStreamSource.h
+++ b/Engine/Plugins/Experimental/MediaStream/Source/MediaStream/Public/MediaStreamSource.h
@@ -29,7 +29,7 @@ struct FMediaStreamSource
 	 * Used by the handler subsystem to identifer the scheme used.
 	 * A value of NAME_None identifies a null or invalid source.
 	 */
-	UPROPERTY(EditInstanceOnly, BlueprintReadOnly, Category = "Media Stream|Source",
+	UPROPERTY(BlueprintReadOnly, Category = "Media Stream|Source",
 		meta = (GetOptions = "/Script/MediaStream.MediaStreamSourceBlueprintFunctionLibrary:GetSchemeTypes"))
 	FName Scheme = NAME_None;
 
@@ -37,13 +37,13 @@ struct FMediaStreamSource
 	 * Path to the referenced source, such as a file/asset path or a managed source name.
 	 * A path must be provided to have a valid scheme.
 	 */
-	UPROPERTY(EditInstanceOnly, BlueprintReadOnly, Category = "Media Stream|Source")
+	UPROPERTY(BlueprintReadOnly, Category = "Media Stream|Source")
 	FString Path;
 
 	/**
 	 * Source of the Media. Either a UMediaStream, UMediaSource or a UMediaPlaylist
 	 */
-	UPROPERTY(EditInstanceOnly, BlueprintReadOnly, Category = "Media Stream|Source",
+	UPROPERTY(BlueprintReadOnly, Category = "Media Stream|Source",
 		meta = (AllowedClasses="MediaStream,MediaSource,MediaPlaylist"))
 	TObjectPtr<UObject> Object = nullptr;
 
diff --git a/Engine/Plugins/Experimental/MediaStream/Source/MediaStream/Public/MediaStreamTextureConfig.h b/Engine/Plugins/Experimental/MediaStream/Source/MediaStream/Public/MediaStreamTextureConfig.h
--- a/Engine/Plugins/Experimental/MediaStream/Source/MediaStream/Public/MediaStreamTextureConfig.h
+++ b/Engine/Plugins/Experimental/MediaStream/Source/MediaStream/Public/MediaStreamTextureConfig.h
@@ -13,7 +13,7 @@ struct FMediaStreamTextureConfig
 {
 	GENERATED_BODY()
 
-	UPROPERTY(EditInstanceOnly, BlueprintReadWrite, Category = "Media Stream|Texture")
+	UPROPERTY(EditInstanceOnly, BlueprintReadWrite, Category = "Media Stream|Texture", DisplayName = "Enable Real Time Mips")
 	bool bEnableMipGen = false;
 
 	MEDIASTREAM_API bool operator==(const FMediaStreamTextureConfig& InOther) const;
diff --git a/Engine/Plugins/Experimental/MediaStream/Source/MediaStreamEditor/Private/DetailsPanel/MediaStreamCustomization.cpp b/Engine/Plugins/Experimental/MediaStream/Source/MediaStreamEditor/Private/DetailsPanel/MediaStreamCustomization.cpp
--- a/Engine/Plugins/Experimental/MediaStream/Source/MediaStreamEditor/Private/DetailsPanel/MediaStreamCustomization.cpp
+++ b/Engine/Plugins/Experimental/MediaStream/Source/MediaStreamEditor/Private/DetailsPanel/MediaStreamCustomization.cpp
@@ -1,119 +1,21 @@
 // Copyright Epic Games, Inc. All Rights Reserved.
 
-#include "MediaStreamCustomization.h"
+#include "DetailsPanel/MediaStreamCustomization.h"
 
 #include "DetailCategoryBuilder.h"
 #include "DetailLayoutBuilder.h"
 #include "DetailWidgetRow.h"
-#include "IDetailGroup.h"
-#include "IMediaStreamPlayer.h"
-#include "MediaPlayerEditorModule.h"
+#include "IDetailChildrenBuilder.h"
 #include "MediaStream.h"
 #include "MediaStreamEditorSequencerLibrary.h"
 #include "MediaStreamEditorStyle.h"
-#include "MediaStreamEnums.h"
-#include "MediaStreamPlayerConfig.h"
-#include "Modules/ModuleManager.h"
-#include "Styling/StyleColors.h"
-#include "Widgets/Input/SButton.h"
-#include "Widgets/Images/SImage.h"
-#include "Widgets/Layout/SBox.h"
+#include "MediaStreamWidgets.h"
 #include "Widgets/Text/STextBlock.h"
 
 #define LOCTEXT_NAMESPACE "MediaStreamCustomization"
 
-class SMediaStreamRefreshSliderWidget : public SCompoundWidget
+namespace UE::MediaStreamEditor
 {
-	SLATE_BEGIN_ARGS(SMediaStreamRefreshSliderWidget) {}
-	SLATE_END_ARGS()
-
-	void Construct(const FArguments& InArgs, UMediaStream* InMediaStream)
-	{
-		MediaStreamWeak = InMediaStream;
-
-		SetCanTick(true);
-
-		ChildSlot
-		[
-			SAssignNew(Content, SBox)
-			.HAlign(HAlign_Fill)
-			.VAlign(VAlign_Fill)
-			[
-				CreateSlider()
-			]
-		];
-	}
-
-	virtual void Tick(const FGeometry& InAllottedGeometry, const double InCurrentTime, const float InDeltaTime) override
-	{
-		SCompoundWidget::Tick(InAllottedGeometry, InCurrentTime, InDeltaTime);
-
-		UMediaPlayer* MediaPlayer = GetMediaPlayer();
-		UMediaPlayer* CurrentMediaPlayer = MediaPlayerWeak.Get();
-
-		if (MediaPlayer != CurrentMediaPlayer)
-		{
-			Content->SetContent(CreateSlider());
-		}
-	}
-
-private:
-	TWeakObjectPtr<UMediaStream> MediaStreamWeak;
-	TWeakObjectPtr<UMediaPlayer> MediaPlayerWeak;
-	TSharedPtr<SBox> Content;
-
-	UMediaPlayer* GetMediaPlayer() const
-	{
-		UMediaStream* MediaStream = MediaStreamWeak.Get();
-
-		if (!MediaStream)
-		{
-			return nullptr;
-		}
-
-		IMediaStreamPlayer* MediaStreamPlayer = MediaStream->GetPlayer().GetInterface();
-
-		if (!MediaStreamPlayer)
-		{
-			return nullptr;
-		}
-
-		return MediaStreamPlayer->GetPlayer();
-	}
-
-	TSharedRef<SWidget> CreateSlider()
-	{
-		UMediaPlayer* MediaPlayer = GetMediaPlayer();
-
-		if (!MediaPlayer)
-		{
-			return SNullWidget::NullWidget;
-		}
-
-		IMediaPlayerEditorModule* MediaPlayerEditorModule = FModuleManager::LoadModulePtr<IMediaPlayerEditorModule>("MediaPlayerEditor");
-
-		if (!MediaPlayerEditorModule)
-		{
-			return SNullWidget::NullWidget;
-		}
-
-		MediaPlayerWeak = MediaPlayer;
-
-		TArray<TWeakObjectPtr<UMediaPlayer>> MediaPlayerList;
-		MediaPlayerList.Add(MediaPlayer);
-
-		const TSharedRef<IMediaPlayerSlider> MediaPlayerSlider = MediaPlayerEditorModule->CreateMediaPlayerSliderWidget(MediaPlayerList);
-
-		MediaPlayerSlider->SetSliderHandleColor(FSlateColor(EStyleColor::AccentBlue));
-		MediaPlayerSlider->SetVisibleWhenInactive(EVisibility::Visible);
-
-		return MediaPlayerSlider;
-	}
-};
-
-FMediaStreamCustomization::FMediaStreamCustomization()
-{
-}
 
 void FMediaStreamCustomization::CustomizeDetails(IDetailLayoutBuilder& InDetailBuilder)
 {
@@ -136,20 +38,50 @@ void FMediaStreamCustomization::CustomizeDetails(IDetailLayoutBuilder& InDetailB
 		return;
 	}
 
-	IDetailCategoryBuilder& MediaStreamCategory = InDetailBuilder.EditCategory("Media Stream");
+	bool bHasValidPlayer = false;
+
+	if (TSharedPtr<IPropertyHandle> PlayerHandle = InDetailBuilder.GetProperty(UMediaStream::GetPlayerPropertyName()))
+	{
+		UObject* Player;
+
+		if (PlayerHandle->GetValue(Player) != FPropertyAccess::Result::Fail)
+		{
+			if (Player)
+			{
+				bHasValidPlayer = true;
+			}
+		}
+	}
+
+	if (bHasValidPlayer)
+	{
+		AddControlCategory(InDetailBuilder);
+	}
 
-	AddControlControls(MediaStreamCategory);
+	AddSourceCategory(InDetailBuilder);
+
+	if (bHasValidPlayer)
+	{
+		AddDetailsCategory(InDetailBuilder);
+		AddTextureCategory(InDetailBuilder);
+		AddCacheCategory(InDetailBuilder);
+		AddPlayerCategory(InDetailBuilder);
+	}
 }
 
-void FMediaStreamCustomization::AddControlControls(IDetailCategoryBuilder& InMediaStreamCategory)
+void FMediaStreamCustomization::AddControlCategory(IDetailLayoutBuilder& InDetailBuilder)
 {
-	IDetailGroup& ControlGroup = InMediaStreamCategory.AddGroup(TEXT("Control"), LOCTEXT("Control", "Control"));
+	UMediaStream* MediaStream = GetMediaStream();
+
+	if (!MediaStream)
+	{
+		return;
+	}
 
-	TWeakPtr<FMediaStreamCustomization> WeakSelf = StaticCastWeakPtr<FMediaStreamCustomization>(AsWeak());
-	const ISlateStyle* Stream = &FMediaStreamEditorStyle::Get().Get();
+	IDetailCategoryBuilder& ControlsCategory = InDetailBuilder.EditCategory("Media Controls");
 
 	// Add media player playback slider
-	ControlGroup.AddWidgetRow()
+	ControlsCategory.AddCustomRow(LOCTEXT("Track", "Track"))
 		.IsEnabled(TAttribute<bool>::CreateSPLambda(this, 
 			[this]()
 			{
@@ -161,591 +93,201 @@ void FMediaStreamCustomization::AddControlControls(IDetailCategoryBuilder& InMed
 				return false;
 			}))
 		[
-			SNew(SMediaStreamRefreshSliderWidget, GetMediaStream())
+			FMediaStreamWidgets::CreateTrackWidget(MediaStream)
 		];
 
 	// Add media control buttons.
-	ControlGroup.AddWidgetRow()
+	ControlsCategory.AddCustomRow(LOCTEXT("Controls", "Controls"))
 	[
-		SNew(SHorizontalBox)
+		FMediaStreamWidgets::CreateControlsWidget(MediaStream)
+	];
 
-		// Open button.
-		+ SHorizontalBox::Slot()
-		.VAlign(VAlign_Center)
-		.AutoWidth()
-		[
-			SNew(SButton)
-			.VAlign(VAlign_Center)
-			.Visibility_Lambda([WeakSelf]
-				{
-					if (const TSharedPtr<FMediaStreamCustomization> Self = WeakSelf.Pin())
-					{
-						if (UMediaStream* MediaStream = Self->GetMediaStream())
-						{
-							if (MediaStream->HasValidSource())
-							{
-								if (IMediaStreamPlayer* MediaStreamPlayer = MediaStream->GetPlayer().GetInterface())
-								{
-									if (UMediaPlayer* MediaPlayer = MediaStreamPlayer->GetPlayer())
-									{
-										if (MediaPlayer->IsClosed())
-										{
-											return EVisibility::Visible;
-										}
-
-										return EVisibility::Collapsed;
-									}
-								}
-
-								return EVisibility::Visible;
-							}
-						}
-					}
-
-					return EVisibility::Collapsed;
-				})
-			.OnClicked_Lambda([WeakSelf]() -> FReply
-				{
-					if (const TSharedPtr<FMediaStreamCustomization> Self = WeakSelf.Pin())
-					{
-						if (UMediaStream* MediaStream = Self->GetMediaStream())
-						{
-							if (IMediaStreamPlayer* MediaStreamPlayer = MediaStream->GetPlayer().GetInterface())
-							{
-								MediaStreamPlayer->OpenSource();
-							}
-						}
-
-						return FReply::Handled();
-					}
-
-					return FReply::Unhandled();
-				})
-			[
-				SNew(SImage)
-				.ColorAndOpacity(FSlateColor::UseForeground())
-				.Image(Stream->GetBrush("MediaStreamEditor.OpenMedia.Small"))
-				.ToolTipText(LOCTEXT("Open Media", "Opens the current media, if the source is valid."))
-			]
-		]
+	TSharedPtr<IPropertyHandle> PlayerHandle = InDetailBuilder.GetProperty(UMediaStream::GetPlayerPropertyName());
+	TSharedPtr<IPropertyHandle> PlayerConfigHandle = PlayerHandle->GetChildHandle("PlayerConfig");
 
-		// Close button.
-		+ SHorizontalBox::Slot()
-		.VAlign(VAlign_Center)
-		.AutoWidth()
-		[
-			SNew(SButton)
-			.VAlign(VAlign_Center)
-			.Visibility_Lambda([WeakSelf]
-				{
-					if (const TSharedPtr<FMediaStreamCustomization> Self = WeakSelf.Pin())
-					{
-						if (UMediaStream* MediaStream = Self->GetMediaStream())
-						{
-							if (MediaStream->HasValidSource())
-							{
-								if (IMediaStreamPlayer* MediaStreamPlayer = MediaStream->GetPlayer().GetInterface())
-								{
-									if (UMediaPlayer* MediaPlayer = MediaStreamPlayer->GetPlayer())
-									{
-										if (!FMediaStreamEditorSequencerLibrary::HasTrack(MediaStream))
-										{
-											if (!MediaPlayer->IsClosed())
-											{
-												return EVisibility::Visible;
-											}
-										}
-									}
-								}
-							}
-						}
-					}
-
-					return EVisibility::Collapsed;
-				})
-			.OnClicked_Lambda([WeakSelf]() -> FReply
-				{
-					if (const TSharedPtr<FMediaStreamCustomization> Self = WeakSelf.Pin())
-					{
-						if (UMediaStream* MediaStream = Self->GetMediaStream())
-						{
-							if (IMediaStreamPlayer* MediaStreamPlayer = MediaStream->GetPlayer().GetInterface())
-							{
-								MediaStreamPlayer->Close();
-							}
-						}
-
-						return FReply::Handled();
-					}
-
-					return FReply::Unhandled();
-				})
-			[
-				SNew(SImage)
-				.ColorAndOpacity(FSlateColor::UseForeground())
-				.Image(Stream->GetBrush("MediaStreamEditor.CloseMedia.Small"))
-				.ToolTipText(LOCTEXT("Close Media", "Closes the currently opened media."))
-			]
-		]
+	if (!PlayerConfigHandle.IsValid())
+	{
+		return;
+	}
 
-		// Rewind button.
-		+ SHorizontalBox::Slot()
-		.VAlign(VAlign_Center)
-		.AutoWidth()
-		[
-			SNew(SButton)
-			.VAlign(VAlign_Center)
-			.IsEnabled_Lambda([WeakSelf]
-				{
-					if (const TSharedPtr<FMediaStreamCustomization> Self = WeakSelf.Pin())
-					{
-						if (UMediaStream* MediaStream = Self->GetMediaStream())
-						{
-							if (!FMediaStreamEditorSequencerLibrary::HasTrack(MediaStream))
-							{
-								if (const UMediaPlayer* MediaPlayer = Self->GetMediaPlayer())
-								{
-									return MediaPlayer->IsReady() &&
-										MediaPlayer->SupportsSeeking() &&
-										MediaPlayer->GetTime() > FTimespan::Zero();
-								}
-							}
-						}
-					}
-
-					return false;
-				})
-			.OnClicked_Lambda([WeakSelf]() -> FReply
-				{
-					if (const TSharedPtr<FMediaStreamCustomization> Self = WeakSelf.Pin())
-					{
-						Self->OnPlaybackSeek(EMediaStreamPlaybackSeek::Previous);
-						return FReply::Handled();
-					}
-
-					return FReply::Unhandled();
-				})
-			[
-				SNew(SImage)
-				.ColorAndOpacity(FSlateColor::UseForeground())
-				.Image(Stream->GetBrush("MediaStreamEditor.RewindMedia.Small"))
-				.ToolTipText(LOCTEXT("Rewind", "Rewind the media to the beginning"))
-			]
-		]
+	if (TSharedPtr<IPropertyHandle> PlayerOnOpenHandle = PlayerHandle->GetChildHandle("bPlayOnOpen"))
+	{
+		ControlsCategory.AddProperty(PlayerOnOpenHandle);
+	}
 
-		// Reverse button.
-		+ SHorizontalBox::Slot()
-		.VAlign(VAlign_Center)
-		.AutoWidth()
-		[
-			SNew(SButton)
-			.VAlign(VAlign_Center)
-			.IsEnabled_Lambda([WeakSelf]
-				{
-					if (const TSharedPtr<FMediaStreamCustomization> Self = WeakSelf.Pin())
-					{
-						if (UMediaStream* MediaStream = Self->GetMediaStream())
-						{
-							if (!FMediaStreamEditorSequencerLibrary::HasTrack(MediaStream))
-							{
-								if (UMediaPlayer* MediaPlayer = Self->GetMediaPlayer())
-								{
-									const float Rate = FMath::Abs(MediaPlayer->GetRate());
-
-									return MediaPlayer->IsReady() && Self->IsRateSupported(Rate > 0 ? -Rate : -1.f);
-								}
-							}
-						}
-					}
-
-					return false;
-				})
-			.OnClicked_Lambda([WeakSelf]() -> FReply
-				{
-					if (const TSharedPtr<FMediaStreamCustomization> Self = WeakSelf.Pin())
-					{
-						Self->OnChangePlaybackDirection(EMediaStreamPlaybackDirection::Backward);
-						return FReply::Handled();
-					}
-
-					return FReply::Unhandled();
-				})
-			[
-				SNew(SImage)
-				.ColorAndOpacity(FSlateColor::UseForeground())
-				.Image(Stream->GetBrush("MediaStreamEditor.ReverseMedia.Small"))
-				.ToolTipText(LOCTEXT("Reverse", "Reverse media playback"))
-			]
-		]
+	if (TSharedPtr<IPropertyHandle> LoopingHandle = PlayerHandle->GetChildHandle("bLooping"))
+	{
+		ControlsCategory.AddProperty(LoopingHandle);
+	}
+}
 
-		// Play button.
-		+ SHorizontalBox::Slot()
-		.VAlign(VAlign_Center)
-		.AutoWidth()
-		[
-			SNew(SButton)
-			.VAlign(VAlign_Center)
-			.IsEnabled_Lambda([WeakSelf]
-				{
-					if (const TSharedPtr<FMediaStreamCustomization> Self = WeakSelf.Pin())
-					{
-						if (UMediaStream* MediaStream = Self->GetMediaStream())
-						{
-							if (!FMediaStreamEditorSequencerLibrary::HasTrack(MediaStream))
-							{
-								// Is the player paused or fast forwarding/rewinding?
-								if (const UMediaPlayer* MediaPlayer = Self->GetMediaPlayer())
-								{
-									return MediaPlayer->IsReady()
-										&& (!MediaPlayer->IsPlaying() || (MediaPlayer->GetRate() != 1.0f));
-								}
-							}
-						}
-					}
-
-					return false;
-				})
-			.Visibility_Lambda([WeakSelf]
-				{
-					if (const TSharedPtr<FMediaStreamCustomization> Self = WeakSelf.Pin())
-					{
-						if (const UMediaPlayer* MediaPlayer = Self->GetMediaPlayer())
-						{
-							return MediaPlayer->IsPlaying() ? EVisibility::Collapsed : EVisibility::Visible;
-						}
-					}
-
-					return EVisibility::Visible;
-				})
-			.OnClicked_Lambda([WeakSelf]() -> FReply
-				{
-					if (const TSharedPtr<FMediaStreamCustomization> Self = WeakSelf.Pin())
-					{
-						Self->OnChangePlaybackState(EMediaStreamPlaybackState::Play);
-						return FReply::Handled();
-					}
-
-					return FReply::Unhandled();
-				})
-			[
-				SNew(SImage)
-				.ColorAndOpacity(FSlateColor::UseForeground())
-				.Image(Stream->GetBrush("MediaStreamEditor.PlayMedia.Small"))
-				.ToolTipText(LOCTEXT("Play", "Start media playback"))
-			]
-		]
+void FMediaStreamCustomization::AddSourceCategory(IDetailLayoutBuilder& InDetailBuilder)
+{
+	TSharedPtr<IPropertyHandle> SourceHandle = InDetailBuilder.GetProperty(UMediaStream::GetSourcePropertyName());
 
-		// Pause button.
-		+ SHorizontalBox::Slot()
-		.VAlign(VAlign_Center)
-		.AutoWidth()
-		[
-			SNew(SButton)
-			.VAlign(VAlign_Center)
-			.IsEnabled_Lambda([WeakSelf]
-				{
-					if (const TSharedPtr<FMediaStreamCustomization> Self = WeakSelf.Pin())
-					{
-						if (UMediaStream* MediaStream = Self->GetMediaStream())
-						{
-							if (!FMediaStreamEditorSequencerLibrary::HasTrack(MediaStream))
-							{
-								if (const UMediaPlayer* MediaPlayer = Self->GetMediaPlayer())
-								{
-									return MediaPlayer->CanPause() && !MediaPlayer->IsPaused();
-								}
-							}
-						}
-					}
-
-					return false;
-				})
-			.Visibility_Lambda([WeakSelf]
-				{
-					if (const TSharedPtr<FMediaStreamCustomization> Self = WeakSelf.Pin())
-					{
-						if (const UMediaPlayer* MediaPlayer = Self->GetMediaPlayer())
-						{
-							const bool bIsVisible = MediaPlayer->CanPause() && !MediaPlayer->IsPaused();
-							return bIsVisible ? EVisibility::Visible : EVisibility::Collapsed;
-						}
-					}
-
-					return EVisibility::Collapsed;
-				})
-			.OnClicked_Lambda([WeakSelf]() -> FReply
-				{
-					if (const TSharedPtr<FMediaStreamCustomization> Self = WeakSelf.Pin())
-					{
-						Self->OnChangePlaybackState(EMediaStreamPlaybackState::Pause);
-						return FReply::Handled();
-					}
-
-					return FReply::Unhandled();
-				})
-			[
-				SNew(SImage)
-				.ColorAndOpacity(FSlateColor::UseForeground())
-				.Image(Stream->GetBrush("MediaStreamEditor.PauseMedia.Small"))
-				.ToolTipText(LOCTEXT("Pause", "Pause media playback"))
-			]
-		]
+	if (!SourceHandle.IsValid())
+	{
+		return;
+	}
 
-		// Forward button.
-		+ SHorizontalBox::Slot()
-		.VAlign(VAlign_Center)
-		.AutoWidth()
-		[
-			SNew(SButton)
-			.VAlign(VAlign_Center)
-			.IsEnabled_Lambda([WeakSelf]
-				{
-					if (const TSharedPtr<FMediaStreamCustomization> Self = WeakSelf.Pin())
-					{
-						if (UMediaStream* MediaStream = Self->GetMediaStream())
-						{
-							if (!FMediaStreamEditorSequencerLibrary::HasTrack(MediaStream))
-							{
-								if (UMediaPlayer* MediaPlayer = Self->GetMediaPlayer())
-								{
-									return MediaPlayer->IsReady();
-								}
-							}
-						}
-					}
-
-					return false;
-				})
-			.OnClicked_Lambda([WeakSelf]() -> FReply
-				{
-					if (const TSharedPtr<FMediaStreamCustomization> Self = WeakSelf.Pin())
-					{
-						Self->OnPlaybackSeek(EMediaStreamPlaybackSeek::End);
-						return FReply::Handled();
-					}
-
-					return FReply::Unhandled();
-				})
-			[
-				SNew(SImage)
-				.ColorAndOpacity(FSlateColor::UseForeground())
-				.Image(Stream->GetBrush("MediaStreamEditor.ForwardMedia.Small"))
-				.RenderTransform(FSlateRenderTransform(TTransform2<float>(TQuat2<float>(PI), FVector2f(20.f, 20.f))))
-				.ToolTipText(LOCTEXT("Forward", "Fast forward media playback"))
-			]
-		]
+	SourceHandle->MarkHiddenByCustomization();
 
-		// Add Media Track
-		+ SHorizontalBox::Slot()
-		.VAlign(VAlign_Center)
-		.AutoWidth()
-		[
-			SNew(SButton)
-			.VAlign(VAlign_Center)
-			.IsEnabled_Lambda([WeakSelf]
-				{
-					if (const TSharedPtr<FMediaStreamCustomization> Self = WeakSelf.Pin())
-					{
-						if (UMediaStream* MediaStream = Self->GetMediaStream())
-						{
-							return FMediaStreamEditorSequencerLibrary::CanAddTrack(MediaStream);
-						}
-					}
-
-					return false;
-				})
-			.OnClicked_Lambda([WeakSelf]() -> FReply
-				{
-					if (const TSharedPtr<FMediaStreamCustomization> Self = WeakSelf.Pin())
-					{
-						if (Self->MediaStreamsList.Num() == 1)
-						{
-							if (UMediaStream* MediaStream = Self->MediaStreamsList[0].Get())
-							{
-								FMediaStreamEditorSequencerLibrary::AddTrack(MediaStream);
-							}
-						}
-					}
-
-					return FReply::Unhandled();
-				})
-			[
-				SNew(STextBlock)
-				.Font(IDetailLayoutBuilder::GetDetailFont())
-				.Text(LOCTEXT("AddTrack", "Add Track"))
-			]
-		]
-	];
-}
+	IDetailCategoryBuilder& SourceCategory = InDetailBuilder.EditCategory("Media Source");
 
-void FMediaStreamCustomization::OnChangePlaybackState(EMediaStreamPlaybackState InState)
-{
-	for (TWeakObjectPtr<UMediaStream>& MediaStreamPtr : MediaStreamsList)
+	TArray<UObject*> Outers;
+	SourceHandle->GetOuterObjects(Outers);
+
+	if (Outers.IsEmpty())
 	{
-		if (UMediaStream* MediaStream = MediaStreamPtr.Get())
-		{
-			if (IMediaStreamPlayer* MediaStreamPlayer = MediaStream->GetPlayer().GetInterface())
-			{
-				if (MediaStreamPlayer->IsReadOnly())
-				{
-					continue;
-				}
+		return;
+	}
 
-				switch (InState)
-				{
-					case EMediaStreamPlaybackState::Play:
-					{
-						const float RequestedRate = FMath::Abs(MediaStreamPlayer->GetPlayerConfig().Rate);
+	UMediaStream* MediaStream = Cast<UMediaStream>(Outers[0]);
 
-						if (RequestedRate != MediaStreamPlayer->GetPlayerConfig().Rate)
-						{
-							FMediaStreamPlayerConfig PlayerConfig = MediaStreamPlayer->GetPlayerConfig();
-							PlayerConfig.Rate = RequestedRate;
+	if (!MediaStream)
+	{
+		return;
+	}
 
-							MediaStreamPlayer->SetPlayerConfig(PlayerConfig);
-						}
+	const IMediaStreamSchemeHandler::FCustomWidgets WidgetRows = FMediaStreamWidgets::GenerateSourceSchemeRows(MediaStream);
 
-						MediaStreamPlayer->Play();
-						break;
-					}
+	if (WidgetRows.CustomRows.IsEmpty())
+	{
+		return;
+	}
 
-					case EMediaStreamPlaybackState::Pause:
-						MediaStreamPlayer->Pause();
-						break;
-				}
-			}
-		}
+	for (const IMediaStreamSchemeHandler::FWidgetRow& WidgetRow : WidgetRows.CustomRows)
+	{
+		SourceCategory.AddCustomRow(WidgetRow.Name)
+			.IsEnabled(WidgetRow.Enabled)
+			.Visibility(WidgetRow.Visibility)
+			.NameContent()
+			[
+				SNew(STextBlock)
+				.Font(IDetailLayoutBuilder::GetDetailFont())
+				.Text(WidgetRow.Name)
+			]
+			.ValueContent()
+			[
+				WidgetRow.Widget
+			];
 	}
 }
 
-void FMediaStreamCustomization::OnChangePlaybackDirection(EMediaStreamPlaybackDirection InDirection)
+void FMediaStreamCustomization::AddDetailsCategory(IDetailLayoutBuilder& InDetailBuilder)
 {
-	for (TWeakObjectPtr<UMediaStream>& MediaStreamPtr : MediaStreamsList)
+	IDetailCategoryBuilder& MediaDetailsCategory = InDetailBuilder.EditCategory("Media Details");
+
+	MediaDetailsCategory.AddCustomRow(LOCTEXT("Details", "Details"))
+		[
+			FMediaStreamWidgets::CreateTextureDetailsWidget(GetMediaStream())
+		];
+}
+
+void FMediaStreamCustomization::AddTextureCategory(IDetailLayoutBuilder& InDetailBuilder)
+{
+	TSharedPtr<IPropertyHandle> PlayerHandle = InDetailBuilder.GetProperty(UMediaStream::GetPlayerPropertyName());
+	TSharedPtr<IPropertyHandle> MediaTextureHandle = PlayerHandle->GetChildHandle("MediaTexture");
+	TSharedPtr<IPropertyHandle> TextureConfigHandle = PlayerHandle->GetChildHandle("TextureConfig");
+
+	if (!MediaTextureHandle.IsValid() && !TextureConfigHandle.IsValid())
 	{
-		if (UMediaStream* MediaStream = MediaStreamPtr.Get())
-		{
-			if (IMediaStreamPlayer* MediaStreamPlayer = MediaStream->GetPlayer().GetInterface())
-			{
-				if (MediaStreamPlayer->IsReadOnly())
-				{
-					continue;
-				}
+		return;
+	}
 
-				const float RequestedRate = FMath::Abs(MediaStreamPlayer->GetPlayerConfig().Rate)
-					* (InDirection == EMediaStreamPlaybackDirection::Forward ? 1.f : -1.f);
+	IDetailCategoryBuilder& MediaTextureCategory = InDetailBuilder.EditCategory("Media Texture");
 
-				if (RequestedRate != MediaStreamPlayer->GetPlayerConfig().Rate)
-				{
-					FMediaStreamPlayerConfig PlayerConfig = MediaStreamPlayer->GetPlayerConfig();
-					PlayerConfig.Rate = RequestedRate;
+	if (MediaTextureHandle.IsValid())
+	{
+		MediaTextureHandle->MarkHiddenByCustomization();
+		MediaTextureCategory.AddProperty(MediaTextureHandle);
+	}
 
-					MediaStreamPlayer->SetPlayerConfig(PlayerConfig);
-				}
-			}
-		}
+	if (TextureConfigHandle.IsValid())
+	{
+		TextureConfigHandle->MarkHiddenByCustomization();
+		MediaTextureCategory.AddProperty(TextureConfigHandle);
 	}
 }
 
-void FMediaStreamCustomization::OnPlaybackSeek(EMediaStreamPlaybackSeek InPosition)
+void FMediaStreamCustomization::AddCacheCategory(IDetailLayoutBuilder& InDetailBuilder)
 {
-	for (TWeakObjectPtr<UMediaStream>& MediaStreamPtr : MediaStreamsList)
+	TSharedPtr<IPropertyHandle> PlayerHandle = InDetailBuilder.GetProperty(UMediaStream::GetPlayerPropertyName());
+	TSharedPtr<IPropertyHandle> PlayerConfigHandle = PlayerHandle->GetChildHandle("PlayerConfig");
+
+	if (!PlayerConfigHandle.IsValid())
 	{
-		if (UMediaStream* MediaStream = MediaStreamPtr.Get())
-		{
-			if (IMediaStreamPlayer* MediaStreamPlayer = MediaStream->GetPlayer().GetInterface())
-			{
-				if (MediaStreamPlayer->IsReadOnly())
-				{
-					continue;
-				}
+		return;
+	}
 
-				switch (InPosition)
-				{
-					case EMediaStreamPlaybackSeek::Previous:
-						MediaStreamPlayer->Previous();
-						break;
+	TSharedPtr<IPropertyHandle> CacheAheadHandle = PlayerConfigHandle->GetChildHandle("CacheAhead");
+	TSharedPtr<IPropertyHandle> CacheBehindHandle = PlayerConfigHandle->GetChildHandle("CacheBehind");
+	TSharedPtr<IPropertyHandle> CacheBehindGameHandle = PlayerConfigHandle->GetChildHandle("CacheBehindGame");
 
-					case EMediaStreamPlaybackSeek::Start:
-						MediaStreamPlayer->Rewind();
-						break;
+	if (!CacheAheadHandle.IsValid() && !CacheBehindHandle.IsValid() && !CacheBehindGameHandle.IsValid())
+	{
+		return;
+	}
 
-					case EMediaStreamPlaybackSeek::End:
-						MediaStreamPlayer->FastForward();
-						break;
+	IDetailCategoryBuilder& MediaCacheCategory = InDetailBuilder.EditCategory("Media Cache");
 
-					case EMediaStreamPlaybackSeek::Next:
-						MediaStreamPlayer->Next();
-						break;
-				}
-			}
-		}
+	if (CacheAheadHandle.IsValid())
+	{
+		CacheAheadHandle->MarkHiddenByCustomization();
+		MediaCacheCategory.AddProperty(CacheAheadHandle);
 	}
-}
 
-UMediaStream* FMediaStreamCustomization::GetMediaStream() const
-{
-	for (const TWeakObjectPtr<UMediaStream>& MediaStreamPtr : MediaStreamsList)
+	if (CacheBehindHandle.IsValid())
 	{
-		if (UMediaStream* MediaStream = MediaStreamPtr.Get())
-		{
-			return MediaStream;
-		}
+		CacheBehindHandle->MarkHiddenByCustomization();
+		MediaCacheCategory.AddProperty(CacheBehindHandle);
 	}
 
-	return nullptr;
+	if (CacheBehindGameHandle.IsValid())
+	{
+		CacheBehindGameHandle->MarkHiddenByCustomization();
+		MediaCacheCategory.AddProperty(CacheBehindGameHandle);
+	}
 }
 
-UMediaPlayer* FMediaStreamCustomization::GetMediaPlayer() const
+void FMediaStreamCustomization::AddPlayerCategory(IDetailLayoutBuilder& InDetailBuilder)
 {
-	for (const TWeakObjectPtr<UMediaStream>& MediaStreamPtr : MediaStreamsList)
+	TSharedPtr<IPropertyHandle> PlayerHandle = InDetailBuilder.GetProperty(UMediaStream::GetPlayerPropertyName());
+	PlayerHandle->MarkHiddenByCustomization();
+
+	IDetailCategoryBuilder& MediaPlayerCategory = InDetailBuilder.EditCategory("Media Player");
+
+	const TArray<FName> PropertyNames = {
+		"RequestedSeekFrame",
+		"PlaybackState",
+		"PlaylistIndex",
+		"PlayerConfig",
+		"bReadOnly"
+	};
+
+	for (const FName& PropertyName : PropertyNames)
 	{
-		if (UMediaStream* MediaStream = MediaStreamPtr.Get())
+		if (TSharedPtr<IPropertyHandle> PropertyHandle = PlayerHandle->GetChildHandle(PropertyName))
 		{
-			if (IMediaStreamPlayer* MediaStreamPlayer = MediaStream->GetPlayer().GetInterface())
-			{				
-				if (UMediaPlayer* MediaPlayer = MediaStreamPlayer->GetPlayer())
-				{
-					return MediaPlayer;
-				}
-			}
+			PropertyHandle->MarkHiddenByCustomization();
+			MediaPlayerCategory.AddProperty(PropertyHandle);
 		}
 	}
-
-	return nullptr;
 }
 
-TArray<TWeakObjectPtr<UMediaPlayer>> FMediaStreamCustomization::GetMediaPlayers() const
+UMediaStream* FMediaStreamCustomization::GetMediaStream() const
 {
-	TArray<TWeakObjectPtr<UMediaPlayer>> MediaPlayers;
-	MediaPlayers.Reserve(MediaStreamsList.Num());
-
 	for (const TWeakObjectPtr<UMediaStream>& MediaStreamPtr : MediaStreamsList)
 	{
 		if (UMediaStream* MediaStream = MediaStreamPtr.Get())
 		{
-			if (IMediaStreamPlayer* MediaStreamPlayer = MediaStream->GetPlayer().GetInterface())
-			{
-				if (UMediaPlayer* MediaPlayer = MediaStreamPlayer->GetPlayer())
-				{
-					MediaPlayers.Add(MediaPlayer);
-				}
-			}
+			return MediaStream;
 		}
 	}
 
-	return MediaPlayers;
+	return nullptr;
 }
 
-bool FMediaStreamCustomization::IsRateSupported(float InRate) const
-{
-	UMediaPlayer* MediaPlayer = GetMediaPlayer();
-
-	if (!MediaPlayer)
-	{
-		return false;
-	}
-
-	return MediaPlayer->SupportsRate(InRate, /* Unthinned */ true);
-}
+} // UE::MediaStreamEditor
 
 #undef LOCTEXT_NAMESPACE
diff --git a/Engine/Plugins/Experimental/MediaStream/Source/MediaStreamEditor/Private/DetailsPanel/MediaStreamCustomization.h b/Engine/Plugins/Experimental/MediaStream/Source/MediaStreamEditor/Private/DetailsPanel/MediaStreamCustomization.h
--- a/Engine/Plugins/Experimental/MediaStream/Source/MediaStreamEditor/Private/DetailsPanel/MediaStreamCustomization.h
+++ b/Engine/Plugins/Experimental/MediaStream/Source/MediaStreamEditor/Private/DetailsPanel/MediaStreamCustomization.h
@@ -7,16 +7,13 @@
 #include "Containers/Array.h"
 #include "UObject/WeakObjectPtr.h"
 
-class IDetailCategoryBuilder;
 class IDetailLayoutBuilder;
-class IPropertyHandle;
-class SWidget;
-class ULevelSequence;
-class UMediaPlayer;
 class UMediaStream;
-enum class EMediaStreamPlaybackDirection : uint8;
-enum class EMediaStreamPlaybackSeek : uint8;
-enum class EMediaStreamPlaybackState : uint8;
+
+namespace UE::MediaStreamEditor
+{
+
+class FMediaStreamSourceCustomization;
 
 /**
  * Implements a details view customization for the UMediaStreamComponent class.
@@ -26,7 +23,6 @@ class FMediaStreamCustomization : public IDetailCustomization
 public:
 	/**
 	 * Creates an instance of this class.
-	 *
 	 * @return The new instance.
 	 */
 	static TSharedRef<IDetailCustomization> MakeInstance()
@@ -34,7 +30,6 @@ public:
 		return MakeShared<FMediaStreamCustomization>();
 	}
 
-	FMediaStreamCustomization();
 	virtual ~FMediaStreamCustomization() override = default;
 
 	//~ Begin IDetailCustomization
@@ -45,31 +40,26 @@ private:
 	/** List of the media streams we are editing. */
 	TArray<TWeakObjectPtr<UMediaStream>> MediaStreamsList;
 
-	/**
-	 * Adds the control... controls.
-	 */
-	void AddControlControls(IDetailCategoryBuilder& InMediaStreamCategory);
-
-	/**
-	* Changes the state of selected media streams.
-	*/
-	void OnChangePlaybackState(EMediaStreamPlaybackState InState);
+	/** Returns the first media stream being edited. */
+	UMediaStream* GetMediaStream() const;
 
-	/**
-	* Changes the state of selected media streams.
-	*/
-	void OnChangePlaybackDirection(EMediaStreamPlaybackDirection InDirection);
+	/** Adds a scrubbable track and media control buttons */
+	void AddControlCategory(IDetailLayoutBuilder& InDetailBuilder);
 
-	/**
-	* Changes the state of selected media streams.
-	*/
-	void OnPlaybackSeek(EMediaStreamPlaybackSeek InPosition);
+	/** Adds options for the media source. */
+	void AddSourceCategory(IDetailLayoutBuilder& InDetailBuilder);
 
-	UMediaStream* GetMediaStream() const;
+	/** Adds media and player details. */
+	void AddDetailsCategory(IDetailLayoutBuilder& InDetailBuilder);
 
-	UMediaPlayer* GetMediaPlayer() const;
+	/** Adds media texture object and options. */
+	void AddTextureCategory(IDetailLayoutBuilder& InDetailBuilder);
 
-	TArray<TWeakObjectPtr<UMediaPlayer>> GetMediaPlayers() const;
+	/** Adds media cache settings. */
+	void AddCacheCategory(IDetailLayoutBuilder& InDetailBuilder);
 
-	bool IsRateSupported(float InRate) const;
+	/** Adds player config options. */
+	void AddPlayerCategory(IDetailLayoutBuilder& InDetailBuilder);
 };
+
+} // UE::MediaStreamEditor
diff --git a/Engine/Plugins/Experimental/MediaStream/Source/MediaStreamEditor/Private/DetailsPanel/MediaStreamSourceCustomization.cpp b/Engine/Plugins/Experimental/MediaStream/Source/MediaStreamEditor/Private/DetailsPanel/MediaStreamSourceCustomization.cpp
--- a/Engine/Plugins/Experimental/MediaStream/Source/MediaStreamEditor/Private/DetailsPanel/MediaStreamSourceCustomization.cpp
+++ b/Engine/Plugins/Experimental/MediaStream/Source/MediaStreamEditor/Private/DetailsPanel/MediaStreamSourceCustomization.cpp
@@ -1,10 +1,11 @@
 // Copyright Epic Games, Inc. All Rights Reserved.
 
-#include "MediaStreamSourceCustomization.h"
+#include "DetailsPanel/MediaStreamSourceCustomization.h"
 
 #include "DetailLayoutBuilder.h"
 #include "DetailWidgetRow.h"
 #include "IDetailChildrenBuilder.h"
+#include "IDetailPropertyRow.h"
 #include "IDetailGroup.h"
 #include "IMediaStreamSchemeHandler.h"
 #include "MediaStream.h"
@@ -16,42 +17,28 @@
 
 #define LOCTEXT_NAMESPACE "MediaStreamSourceCustomization"
 
-FMediaStreamSourceCustomization::FMediaStreamSourceCustomization()
+namespace UE::MediaStreamEditor
 {
-}
-
-void FMediaStreamSourceCustomization::CustomizeHeader(TSharedRef<IPropertyHandle> InPropertyHandle, FDetailWidgetRow& InHeaderRow,
-	IPropertyTypeCustomizationUtils& InCustomizationUtils)
-{
-}
 
-void FMediaStreamSourceCustomization::CustomizeChildren(TSharedRef<IPropertyHandle> InPropertyHandle, IDetailChildrenBuilder& InChildBuilder,
-	IPropertyTypeCustomizationUtils& InCustomizationUtils)
+IMediaStreamSchemeHandler::FCustomWidgets FMediaStreamSourceCustomization::Customize(UMediaStream* InMediaStream)
 {
-	PropertyHandle = InPropertyHandle;
-
-	if (!FMediaStreamSchemeHandlerLibrary::GetStreamSourcePtr(InPropertyHandle))
+	if (!IsValid(InMediaStream))
 	{
-		return;
+		return {};
 	}
 
-	IDetailGroup& SourceGroup = InChildBuilder.AddGroup(TEXT("Source"), LOCTEXT("Source", "Source"));
+	MediaStream = InMediaStream;
 
-	AddSourceSchemeSelector(SourceGroup, InCustomizationUtils);
-	AddSchemeCustomizations(InPropertyHandle, SourceGroup, InCustomizationUtils);
-}
+	IMediaStreamSchemeHandler::FCustomWidgets CustomWidgets;
 
-void FMediaStreamSourceCustomization::AddSourceSchemeSelector(IDetailGroup& InSourceGroup, IPropertyTypeCustomizationUtils& InCustomizationUtils)
-{
-	FDetailWidgetRow& SourceSchemeRow = InSourceGroup.AddWidgetRow();
+	AddSourceSchemeSelector(CustomWidgets);
+	AddSchemeCustomizations(CustomWidgets);
 
-	SourceSchemeRow.NameContent()
-	[
-		SNew(STextBlock)
-		.Font(IDetailLayoutBuilder::GetDetailFont())
-		.Text(LOCTEXT("SourceSchemeName", "Source Scheme"))
-	];
+	return CustomWidgets;
+}
 
+void FMediaStreamSourceCustomization::AddSourceSchemeSelector(IMediaStreamSchemeHandler::FCustomWidgets& InOutCustomWidgets)
+{
 	TSharedRef<SGridPanel> SchemeSelector = MakeShared<SGridPanel>();
 
 	AddSourceSchemeSelectRow(SchemeSelector, 0, TEXT("None"));
@@ -71,14 +58,16 @@ void FMediaStreamSourceCustomization::AddSourceSchemeSelector(IDetailGroup& InSo
 	}
 	 */
 
-	SourceSchemeRow.ValueContent()
-	[
+	InOutCustomWidgets.CustomRows.Add({
+		LOCTEXT("SourceSchemeName", "Source Scheme"),
 		SchemeSelector
-	];
+	});
 }
 
 void FMediaStreamSourceCustomization::AddSourceSchemeSelectRow(TSharedRef<SGridPanel> InContainer, int32 InRow, const FName& InName)
 {
+	TWeakObjectPtr<UMediaStream> MediaStreamWeak = MediaStream;
+
 	InContainer->AddSlot(0, InRow)
 		.HAlign(EHorizontalAlignment::HAlign_Left)
 		.VAlign(EVerticalAlignment::VAlign_Center)
@@ -86,8 +75,8 @@ void FMediaStreamSourceCustomization::AddSourceSchemeSelectRow(TSharedRef<SGridP
 		[
 			SNew(SCheckBox)
 			.Style(FAppStyle::Get(), "Menu.RadioButton")
-			.IsChecked(this, &FMediaStreamSourceCustomization::GetSourceCheckBoxState, InName)
-			.OnCheckStateChanged(this, &FMediaStreamSourceCustomization::OnSourceCheckBoxStateChanged, InName)
+			.IsChecked_Static(&FMediaStreamSourceCustomization::GetSourceCheckBoxState, MediaStreamWeak, InName)
+			.OnCheckStateChanged_Static(&FMediaStreamSourceCustomization::OnSourceCheckBoxStateChanged, MediaStreamWeak, InName)
 		];
 
 	InContainer->AddSlot(1, InRow)
@@ -108,41 +97,38 @@ void FMediaStreamSourceCustomization::AddSourceSchemeSelectRow(TSharedRef<SGridP
 		];
 }
 
-ECheckBoxState FMediaStreamSourceCustomization::GetSourceCheckBoxState(FName InSourceScheme) const
-{
-	return FMediaStreamSchemeHandlerLibrary::GetScheme(PropertyHandle) == InSourceScheme
-		? ECheckBoxState::Checked
-		: ECheckBoxState::Unchecked;
-}
-
-void FMediaStreamSourceCustomization::OnSourceCheckBoxStateChanged(ECheckBoxState InState, FName InSourceScheme)
+ECheckBoxState FMediaStreamSourceCustomization::GetSourceCheckBoxState(TWeakObjectPtr<UMediaStream> InMediaStreamWeak, FName InSourceScheme)
 {
-	if (PropertyHandle.IsValid())
+	if (UMediaStream* MediaStream = InMediaStreamWeak.Get())
 	{
-		FMediaStreamSchemeHandlerLibrary::SetSource(PropertyHandle.ToSharedRef(), InSourceScheme, "");
+		if (MediaStream->GetSource().Scheme == InSourceScheme )
+		{
+			return ECheckBoxState::Checked;
+		}
 	}
+
+	return ECheckBoxState::Unchecked;
 }
 
-FReply FMediaStreamSourceCustomization::OnSourceSchemeButtonClicked(const FName& InSourceScheme)
+void FMediaStreamSourceCustomization::OnSourceCheckBoxStateChanged(ECheckBoxState InState, TWeakObjectPtr<UMediaStream> InMediaStreamWeak, FName InSourceScheme)
 {
-	if (PropertyHandle.IsValid())
+	if (UMediaStream* MediaStream = InMediaStreamWeak.Get())
 	{
-		FMediaStreamSchemeHandlerLibrary::SetSource(PropertyHandle.ToSharedRef(), InSourceScheme, "");
+		FMediaStreamSchemeHandlerLibrary::SetSource(MediaStream, InSourceScheme, "");
 	}
-
-	return FReply::Handled();
 }
 
-void FMediaStreamSourceCustomization::AddSchemeCustomizations(TSharedRef<IPropertyHandle> InPropertyHandle, IDetailGroup& InSourceGroup,
-	IPropertyTypeCustomizationUtils& InCustomizationUtils)
+void FMediaStreamSourceCustomization::AddSchemeCustomizations(IMediaStreamSchemeHandler::FCustomWidgets& InOutCustomWidgets)
 {
 	const FMediaStreamSchemeHandlerManager& Manager = FMediaStreamSchemeHandlerManager::Get();
 	const TArray<FName> Schemes = Manager.GetSchemeHandlerNames();
 
 	for (const FName& Scheme : Schemes)
 	{
-		Manager.GetHandlerTypeForScheme(Scheme)->CreatePropertyCustomization(InPropertyHandle, InSourceGroup, InCustomizationUtils);
+		Manager.GetHandlerTypeForScheme(Scheme)->CreatePropertyCustomization(MediaStream, InOutCustomWidgets);
 	}
 }
 
+} // UE::MediaStreamEditor
+
 #undef LOCTEXT_NAMESPACE
diff --git a/Engine/Plugins/Experimental/MediaStream/Source/MediaStreamEditor/Private/DetailsPanel/MediaStreamSourceCustomization.h b/Engine/Plugins/Experimental/MediaStream/Source/MediaStreamEditor/Private/DetailsPanel/MediaStreamSourceCustomization.h
--- a/Engine/Plugins/Experimental/MediaStream/Source/MediaStreamEditor/Private/DetailsPanel/MediaStreamSourceCustomization.h
+++ b/Engine/Plugins/Experimental/MediaStream/Source/MediaStreamEditor/Private/DetailsPanel/MediaStreamSourceCustomization.h
@@ -4,13 +4,10 @@
 
 #include "IPropertyTypeCustomization.h"
 
+#include "IMediaStreamSchemeHandler.h"
 #include "Templates/SharedPointer.h"
 
 class FReply;
-class IDetailCategoryBuilder;
-class IDetailGroup;
-class IDetailLayoutBuilder;
-class IPropertyHandle;
 class SGridPanel;
 class SWidget;
 class UMediaPlayer;
@@ -18,41 +15,39 @@ class UMediaStream;
 enum class ECheckBoxState : uint8;
 struct FMediaStreamSource;
 
+namespace UE::MediaStreamEditor
+{
+
 /**
- * Implements a details view customization for the UMediaStreamComponent class.
+ * Generates custom widgets for source scheme types.
+ * Object should not last behind the local scope.
  */
-class FMediaStreamSourceCustomization : public IPropertyTypeCustomization
+class FMediaStreamSourceCustomization : public TSharedFromThis<FMediaStreamSourceCustomization>
 {
 public:
-	/**
-	 * Creates an instance of this class.
-	 *
-	 * @return The new instance.
-	 */
-	static TSharedRef<IPropertyTypeCustomization> MakeInstance()
-	{
-		return MakeShared<FMediaStreamSourceCustomization>();
-	}
-
-	FMediaStreamSourceCustomization();
-	virtual ~FMediaStreamSourceCustomization() override = default;
-
-	//~ Begin IPropertyTypeCustomization
-	virtual void CustomizeHeader(TSharedRef<IPropertyHandle> InPropertyHandle, FDetailWidgetRow& InHeaderRow, IPropertyTypeCustomizationUtils& InCustomizationUtils) override;
-	virtual void CustomizeChildren(TSharedRef<IPropertyHandle> InPropertyHandle, IDetailChildrenBuilder& InChildBuilder, IPropertyTypeCustomizationUtils& InCustomizationUtils) override;
-	//~ End IPropertyTypeCustomization
+	virtual ~FMediaStreamSourceCustomization() = default;
+
+	/** Produces widgets for scheme customizations */
+	IMediaStreamSchemeHandler::FCustomWidgets Customize(UMediaStream* InMediaStream);
 
 private:
-	TSharedPtr<IPropertyHandle> PropertyHandle;
+	/** Whether a particular scheme is active. */
+	static ECheckBoxState GetSourceCheckBoxState(TWeakObjectPtr<UMediaStream> InMediaStreamWeak, FName InSourceScheme);
 
-	void AddSourceSchemeSelector(IDetailGroup& InSourceGroup, IPropertyTypeCustomizationUtils& InCustomizationUtils);
+	/** Changes the source to a new scheme. */
+	static void OnSourceCheckBoxStateChanged(ECheckBoxState InState, TWeakObjectPtr<UMediaStream> InMediaStreamWeak, FName InSourceScheme);
 
-	void AddSourceSchemeSelectRow(TSharedRef<SGridPanel> InContainer, int32 InRow, const FName& InName);
+	/** Temporarily stored media stream. */
+	UMediaStream* MediaStream = nullptr;
 
-	ECheckBoxState GetSourceCheckBoxState(FName InSourceScheme) const;
-	void OnSourceCheckBoxStateChanged(ECheckBoxState InState, FName InSourceScheme);
+	/** Adds the scheme selector widget. @see AddSourceSchemeSelectRow. */
+	void AddSourceSchemeSelector(IMediaStreamSchemeHandler::FCustomWidgets& InOutCustomWidgets);
 
-	FReply OnSourceSchemeButtonClicked(const FName& InSourceScheme);
+	/** Adds the widget for a particular scheme type. */
+	void AddSourceSchemeSelectRow(TSharedRef<SGridPanel> InContainer, int32 InRow, const FName& InName);
 
-	void AddSchemeCustomizations(TSharedRef<IPropertyHandle> InPropertyHandle, IDetailGroup& InSourceGroup, IPropertyTypeCustomizationUtils& InCustomizationUtils);
+	/** Adds a customization for a particular scheme. */
+	void AddSchemeCustomizations(IMediaStreamSchemeHandler::FCustomWidgets& InOutCustomWidgets);
 };
+
+} // UE::MediaStreamEditor
diff --git a/Engine/Plugins/Experimental/MediaStream/Source/MediaStreamEditor/Private/DetailsPanel/SMediaStreamMediaDetails.cpp b/Engine/Plugins/Experimental/MediaStream/Source/MediaStreamEditor/Private/DetailsPanel/SMediaStreamMediaDetails.cpp
new file mode 100644
--- /dev/null
+++ b/Engine/Plugins/Experimental/MediaStream/Source/MediaStreamEditor/Private/DetailsPanel/SMediaStreamMediaDetails.cpp
@@ -0,0 +1,263 @@
+// Copyright Epic Games, Inc. All Rights Reserved.
+
+#include "DetailsPanel/SMediaStreamMediaDetails.h"
+
+#include "IMediaStreamPlayer.h"
+#include "Internationalization/Text.h"
+#include "MediaPlayer.h"
+#include "MediaStream.h"
+#include "Styling/AppStyle.h"
+#include "Widgets/SBoxPanel.h"
+#include "Widgets/Text/STextBlock.h"
+
+#define LOCTEXT_NAMESPACE "SMediaStreamMediaDetails"
+
+namespace UE::MediaStreamEditor
+{
+
+void SMediaStreamMediaDetails::Construct(const FArguments& InArgs, UMediaStream* InMediaStream)
+{
+	MediaStreamWeak = InMediaStream;
+
+	const ISlateStyle& SlateStyle = FAppStyle::Get();
+	const FName StyleName = TEXT("SmallText");
+
+	ChildSlot
+	[
+		SNew(SHorizontalBox)
+		.Visibility(this, &SMediaStreamMediaDetails::AreDetailsVisible)
+
+		// Left side.
+		+ SHorizontalBox::Slot()
+		[
+			SNew(SVerticalBox)
+
+			// Player name.
+			+ SVerticalBox::Slot()
+			.AutoHeight()
+			.VAlign(VAlign_Center)
+			.Padding(4.0f)
+			[
+				SAssignNew(MediaPlayerName, STextBlock)
+				.TextStyle(SlateStyle, StyleName)
+			]
+
+			// Resolution.
+			+ SVerticalBox::Slot()
+			.AutoHeight()
+			.VAlign(VAlign_Center)
+			.Padding(4.0f)
+			[
+				SAssignNew(ResolutionText, STextBlock)
+				.TextStyle(SlateStyle, StyleName)
+			]
+
+			// Frame rate.
+			+ SVerticalBox::Slot()
+			.AutoHeight()
+			.VAlign(VAlign_Center)
+			.Padding(4.0f)
+			[
+				SAssignNew(FrameRateText, STextBlock)
+				.TextStyle(SlateStyle, StyleName)
+			]
+
+			// Resource size.
+			+ SVerticalBox::Slot()
+			.AutoHeight()
+			.VAlign(VAlign_Center)
+			.Padding(4.0f)
+			[
+				SAssignNew(ResourceSizeText, STextBlock)
+				.TextStyle(SlateStyle, StyleName)
+			]
+
+			// Method.
+			+ SVerticalBox::Slot()
+			.AutoHeight()
+			.VAlign(VAlign_Center)
+			.Padding(4.0f)
+			[
+				SAssignNew(MethodText, STextBlock)
+				.TextStyle(SlateStyle, StyleName)
+			]
+		]
+
+		// Right side.
+		+ SHorizontalBox::Slot()
+		[
+			SNew(SVerticalBox)
+
+			// Format.
+			+ SVerticalBox::Slot()
+			.AutoHeight()
+			.VAlign(VAlign_Center)
+			.Padding(4.0f)
+			[
+				SAssignNew(FormatText, STextBlock)
+				.TextStyle(SlateStyle, StyleName)
+			]
+
+			// LOD bias.
+			+ SVerticalBox::Slot()
+			.AutoHeight()
+			.VAlign(VAlign_Center)
+			.Padding(4.0f)
+			[
+				SAssignNew(LODBiasText, STextBlock)
+				.TextStyle(SlateStyle, StyleName)
+			]
+
+			// Num mips.
+			+ SVerticalBox::Slot()
+			.AutoHeight()
+			.VAlign(VAlign_Center)
+			.Padding(4.0f)
+			[
+				SAssignNew(NumMipsText, STextBlock)
+				.TextStyle(SlateStyle, StyleName)
+			]
+
+			// Num tiles.
+			+ SVerticalBox::Slot()
+			.AutoHeight()
+			.VAlign(VAlign_Center)
+			.Padding(4.0f)
+			[
+				SAssignNew(NumTilesText, STextBlock)
+				.TextStyle(SlateStyle, StyleName)
+			]
+		]
+
+	];
+
+	UpdateDetails();
+}
+
+void SMediaStreamMediaDetails::Tick(const FGeometry& AllottedGeometry, const double InCurrentTime, const float InDeltaTime)
+{
+	SCompoundWidget::Tick(AllottedGeometry, InCurrentTime, InDeltaTime);
+
+	UpdateDetails();
+}
+
+UMediaPlayer* SMediaStreamMediaDetails::GetMediaPlayer() const
+{
+	UMediaStream* MediaStream = MediaStreamWeak.Get();
+
+	if (!MediaStream)
+	{
+		return nullptr;
+	}
+
+	IMediaStreamPlayer* MediaStreamPlayer = MediaStream->GetPlayer().GetInterface();
+
+	if (!MediaStreamPlayer)
+	{
+		return nullptr;
+	}
+
+	return MediaStreamPlayer->GetPlayer();
+}
+
+UMediaTexture* SMediaStreamMediaDetails::GetMediaTexture() const
+{
+	UMediaStream* MediaStream = MediaStreamWeak.Get();
+
+	if (!MediaStream)
+	{
+		return nullptr;
+	}
+
+	IMediaStreamPlayer* MediaStreamPlayer = MediaStream->GetPlayer().GetInterface();
+
+	if (!MediaStreamPlayer)
+	{
+		return nullptr;
+	}
+
+	return MediaStreamPlayer->GetMediaTexture();
+}
+
+EVisibility SMediaStreamMediaDetails::AreDetailsVisible() const
+{
+	return (!!GetMediaPlayer() || !!GetMediaTexture())
+		? EVisibility::Visible
+		: EVisibility::Collapsed;
+}
+
+void SMediaStreamMediaDetails::UpdateDetails()
+{
+	FName PlayerName;
+	FString Format;
+	float FrameRate = 0.0f;
+	int32 LODBias = 0;
+	FText Method;
+	int32 NumMips = 0;
+	int32 NumTotalTiles = 0;
+	int64 ResourceSize = 0;
+	int32 SurfaceWidth = 0;
+	int32 SurfaceHeight = 0;
+
+	UMediaPlayer* MediaPlayer = GetMediaPlayer();
+	UMediaTexture* MediaTexture = GetMediaTexture();
+
+	if (!MediaPlayer && !MediaTexture)
+	{
+		return;
+	}
+
+	if (MediaPlayer)
+	{
+		PlayerName = MediaPlayer->GetPlayerName();
+		FrameRate = MediaPlayer->GetVideoTrackFrameRate(INDEX_NONE, INDEX_NONE);
+		Format = MediaPlayer->GetVideoTrackType(INDEX_NONE, INDEX_NONE);
+		FIntPoint NumTiles(EForceInit::ForceInitToZero);
+		MediaPlayer->GetMediaInfo<FIntPoint>(NumTiles, UMediaPlayer::MediaInfoNameSourceNumTiles.Resolve());
+		NumTotalTiles = NumTiles.X * NumTiles.Y;
+	}
+
+	if (MediaTexture)
+	{
+		LODBias = MediaTexture->GetCachedLODBias();
+		Method = MediaTexture->IsCurrentlyVirtualTextured() ?
+			LOCTEXT("MethodVirtualStreamed", "Virtual Streamed")
+			: (!MediaTexture->IsStreamable() ? LOCTEXT("QuickInfo_MethodNotStreamed", "Not Streamed")
+				: LOCTEXT("MethodStreamed", "Streamed"));
+		NumMips = MediaTexture->GetTextureNumMips();
+		ResourceSize = (MediaTexture->GetResourceSizeBytes(EResourceSizeMode::Exclusive) + 512) / 1024;
+		SurfaceWidth = static_cast<int32>(MediaTexture->GetSurfaceWidth());
+		SurfaceHeight = static_cast<int32>(MediaTexture->GetSurfaceHeight());
+	}
+
+	// Update text.
+	MediaPlayerName->SetText(FText::Format(LOCTEXT("Player", "Player: {0}"),
+		FText::FromName(PlayerName)));
+
+	FormatText->SetText(FText::Format(LOCTEXT("Format", "Format: {0}"),
+		FText::FromString(Format)));
+
+	FrameRateText->SetText(FText::Format(LOCTEXT("FrameRate", "Frame Rate: {0}"), 
+		FText::AsNumber(FrameRate)));
+
+	LODBiasText->SetText(FText::Format(LOCTEXT("LODBias", "Combined LOD Bias: {0}"),
+		FText::AsNumber(LODBias)));
+
+	MethodText->SetText(FText::Format(LOCTEXT("Method", "Method: {0}"), Method));
+
+	NumMipsText->SetText(FText::Format(LOCTEXT("NumberOfMips", "Mips: {0}"),
+		FText::AsNumber(NumMips)));
+
+	NumTilesText->SetText(FText::Format(LOCTEXT("NumberOfTiles", "Tiles: {0}"),
+		FText::AsNumber(NumTotalTiles)));
+
+	ResolutionText->SetText(FText::Format(LOCTEXT("Resolution", "Resolution: {0}x{1}"),
+		FText::AsNumber(SurfaceWidth), FText::AsNumber(SurfaceHeight)));
+
+	ResourceSizeText->SetText(FText::Format(LOCTEXT("ResourceSize", "Resource Size: {0} KB"),
+		FText::AsNumber(ResourceSize)));
+}
+
+} // UE::MediaStreamEditor
+
+#undef LOCTEXT_NAMESPACE
diff --git a/Engine/Plugins/Experimental/MediaStream/Source/MediaStreamEditor/Private/DetailsPanel/SMediaStreamMediaDetails.h b/Engine/Plugins/Experimental/MediaStream/Source/MediaStreamEditor/Private/DetailsPanel/SMediaStreamMediaDetails.h
new file mode 100644
--- /dev/null
+++ b/Engine/Plugins/Experimental/MediaStream/Source/MediaStreamEditor/Private/DetailsPanel/SMediaStreamMediaDetails.h
@@ -0,0 +1,60 @@
+// Copyright Epic Games, Inc. All Rights Reserved.
+
+#pragma once
+
+#include "Widgets/SCompoundWidget.h"
+
+#include "UObject/WeakObjectPtr.h"
+
+class SBox;
+class STextBlock;
+class UMediaPlayer;
+class UMediaStream;
+class UMediaTexture;
+
+namespace UE::MediaStreamEditor
+{
+
+/**
+ * Displays details about the media texture and its player.
+ */
+class SMediaStreamMediaDetails : public SCompoundWidget
+{
+	SLATE_BEGIN_ARGS(SMediaStreamMediaDetails)
+		{}
+	SLATE_END_ARGS()
+
+	MEDIASTREAMEDITOR_API void Construct(const FArguments& InArgs, UMediaStream* InMediaStream);
+
+	//~ Begin SWidget
+	MEDIASTREAMEDITOR_API virtual void Tick(const FGeometry& AllottedGeometry, const double InCurrentTime, const float InDeltaTime) override;
+	//~ End SWidget
+
+protected:
+	TWeakObjectPtr<UMediaStream> MediaStreamWeak;
+
+	/** Our widgets. */
+	TSharedPtr<STextBlock> MediaPlayerName;
+	TSharedPtr<STextBlock> FormatText;
+	TSharedPtr<STextBlock> FrameRateText;
+	TSharedPtr<STextBlock> LODBiasText;
+	TSharedPtr<STextBlock> MethodText;
+	TSharedPtr<STextBlock> NumMipsText;
+	TSharedPtr<STextBlock> NumTilesText;
+	TSharedPtr<STextBlock> ResolutionText;
+	TSharedPtr<STextBlock> ResourceSizeText;
+
+	/** Retrieves the media player. */
+	UMediaPlayer* GetMediaPlayer() const;
+
+	/** Retrieves the media texture. */
+	UMediaTexture* GetMediaTexture() const;
+
+	/** Determines whether the entire widget should be visible. */
+	EVisibility AreDetailsVisible() const;
+
+	/** Updates our widgets to reflect the current state. */
+	void UpdateDetails();
+};
+
+} // UE::MediaStreamEditor
diff --git a/Engine/Plugins/Experimental/MediaStream/Source/MediaStreamEditor/Private/DetailsPanel/SMediaStreamMediaTrack.cpp b/Engine/Plugins/Experimental/MediaStream/Source/MediaStreamEditor/Private/DetailsPanel/SMediaStreamMediaTrack.cpp
new file mode 100644
--- /dev/null
+++ b/Engine/Plugins/Experimental/MediaStream/Source/MediaStreamEditor/Private/DetailsPanel/SMediaStreamMediaTrack.cpp
@@ -0,0 +1,94 @@
+// Copyright Epic Games, Inc. All Rights Reserved.
+
+#include "DetailsPanel/SMediaStreamMediaTrack.h"
+
+#include "IMediaStreamPlayer.h"
+#include "MediaPlayer.h"
+#include "MediaPlayerEditorModule.h"
+#include "MediaStream.h"
+#include "Modules/ModuleManager.h"
+#include "Styling/StyleColors.h"
+#include "Widgets/Layout/SBox.h"
+
+namespace UE::MediaStreamEditor
+{
+
+void SMediaStreamMediaTrack::Construct(const FArguments& InArgs, UMediaStream* InMediaStream)
+{
+	MediaStreamWeak = InMediaStream;
+
+	SetCanTick(true);
+
+	ChildSlot
+	[
+		SAssignNew(Content, SBox)
+		.HAlign(HAlign_Fill)
+		.VAlign(VAlign_Fill)
+		[
+			CreateTrack()
+		]
+	];
+}
+
+void SMediaStreamMediaTrack::Tick(const FGeometry& InAllottedGeometry, const double InCurrentTime, const float InDeltaTime)
+{
+	SCompoundWidget::Tick(InAllottedGeometry, InCurrentTime, InDeltaTime);
+
+	UMediaPlayer* MediaPlayer = GetMediaPlayer();
+	UMediaPlayer* CurrentMediaPlayer = MediaPlayerWeak.Get();
+
+	if (MediaPlayer != CurrentMediaPlayer)
+	{
+		Content->SetContent(CreateTrack());
+	}
+}
+
+UMediaPlayer* SMediaStreamMediaTrack::GetMediaPlayer() const
+{
+	UMediaStream* MediaStream = MediaStreamWeak.Get();
+
+	if (!MediaStream)
+	{
+		return nullptr;
+	}
+
+	IMediaStreamPlayer* MediaStreamPlayer = MediaStream->GetPlayer().GetInterface();
+
+	if (!MediaStreamPlayer)
+	{
+		return nullptr;
+	}
+
+	return MediaStreamPlayer->GetPlayer();
+}
+
+TSharedRef<SWidget> SMediaStreamMediaTrack::CreateTrack()
+{
+	UMediaPlayer* MediaPlayer = GetMediaPlayer();
+
+	if (!MediaPlayer)
+	{
+		return SNullWidget::NullWidget;
+	}
+
+	IMediaPlayerEditorModule* MediaPlayerEditorModule = FModuleManager::LoadModulePtr<IMediaPlayerEditorModule>("MediaPlayerEditor");
+
+	if (!MediaPlayerEditorModule)
+	{
+		return SNullWidget::NullWidget;
+	}
+
+	MediaPlayerWeak = MediaPlayer;
+
+	TArray<TWeakObjectPtr<UMediaPlayer>> MediaPlayerList;
+	MediaPlayerList.Add(MediaPlayer);
+
+	const TSharedRef<IMediaPlayerSlider> MediaPlayerSlider = MediaPlayerEditorModule->CreateMediaPlayerSliderWidget(MediaPlayerList);
+
+	MediaPlayerSlider->SetSliderHandleColor(FSlateColor(EStyleColor::AccentBlue));
+	MediaPlayerSlider->SetVisibleWhenInactive(EVisibility::Visible);
+
+	return MediaPlayerSlider;
+}
+
+} // UE::MediaStreamEditor
diff --git a/Engine/Plugins/Experimental/MediaStream/Source/MediaStreamEditor/Private/DetailsPanel/SMediaStreamMediaTrack.h b/Engine/Plugins/Experimental/MediaStream/Source/MediaStreamEditor/Private/DetailsPanel/SMediaStreamMediaTrack.h
new file mode 100644
--- /dev/null
+++ b/Engine/Plugins/Experimental/MediaStream/Source/MediaStreamEditor/Private/DetailsPanel/SMediaStreamMediaTrack.h
@@ -0,0 +1,48 @@
+// Copyright Epic Games, Inc. All Rights Reserved.
+
+#pragma once
+
+#include "Widgets/SCompoundWidget.h"
+
+#include "UObject/WeakObjectPtr.h"
+
+class SBox;
+class UMediaPlayer;
+class UMediaStream;
+
+namespace UE::MediaStreamEditor
+{
+
+/**
+ * Shows a scrubbable track for controlling media playback position.
+ */
+class SMediaStreamMediaTrack : public SCompoundWidget
+{
+	SLATE_BEGIN_ARGS(SMediaStreamMediaTrack)
+		{}
+	SLATE_END_ARGS()
+
+	MEDIASTREAMEDITOR_API void Construct(const FArguments& InArgs, UMediaStream* InMediaStream);
+
+	//~ Begin SWidget
+	MEDIASTREAMEDITOR_API virtual void Tick(const FGeometry& InAllottedGeometry, const double InCurrentTime, const float InDeltaTime) override;
+	//~ End SWidget
+
+private:
+	/** The media stream for this track. */
+	TWeakObjectPtr<UMediaStream> MediaStreamWeak;
+
+	/** The active media player. If this changes, recreate the track. */
+	TWeakObjectPtr<UMediaPlayer> MediaPlayerWeak;
+
+	/** The container for the track. */
+	TSharedPtr<SBox> Content;
+
+	/** Returns the currently active media player. May different from @see MediaPlayerWeak. */
+	UMediaPlayer* GetMediaPlayer() const;
+
+	/** Creates the track widget. */
+	TSharedRef<SWidget> CreateTrack();
+};
+
+} // UE::MediaStreamEditor
diff --git a/Engine/Plugins/Experimental/MediaStream/Source/MediaStreamEditor/Private/DetailsPanel/SMediaStreamPlaybackControls.cpp b/Engine/Plugins/Experimental/MediaStream/Source/MediaStreamEditor/Private/DetailsPanel/SMediaStreamPlaybackControls.cpp
new file mode 100644
--- /dev/null
+++ b/Engine/Plugins/Experimental/MediaStream/Source/MediaStreamEditor/Private/DetailsPanel/SMediaStreamPlaybackControls.cpp
@@ -0,0 +1,541 @@
+// Copyright Epic Games, Inc. All Rights Reserved.
+
+#include "DetailsPanel/SMediaStreamPlaybackControls.h"
+
+#include "DetailLayoutBuilder.h"
+#include "IMediaStreamPlayer.h"
+#include "MediaPlayer.h"
+#include "MediaStream.h"
+#include "MediaStreamEditorSequencerLibrary.h"
+#include "MediaStreamEditorStyle.h"
+#include "MediaStreamEnums.h"
+#include "MediaStreamPlayerConfig.h"
+#include "Styling/AppStyle.h"
+#include "Styling/StyleColors.h"
+#include "Widgets/Images/SImage.h"
+#include "Widgets/Input/SButton.h"
+#include "Widgets/SBoxPanel.h"
+#include "Widgets/Text/STextBlock.h"
+
+#define LOCTEXT_NAMESPACE "SMediaStreamPlaybackControls"
+
+namespace UE::MediaStreamEditor
+{
+
+void SMediaStreamPlaybackControls::Construct(const FArguments& InArgs, UMediaStream* InMediaStream)
+{
+	MediaStreamWeak = InMediaStream;
+
+	const ISlateStyle* MediaStreamStyle = &FMediaStreamEditorStyle::Get().Get();
+
+	ChildSlot
+	[
+		SNew(SHorizontalBox)
+
+		// Open button.
+		+ SHorizontalBox::Slot()
+		.VAlign(VAlign_Center)
+		.AutoWidth()
+		[
+			SNew(SButton)
+			.VAlign(VAlign_Center)
+			.Visibility(this, &SMediaStreamPlaybackControls::Open_GetVisibility)
+			.OnClicked(this, &SMediaStreamPlaybackControls::Open_OnClicked)
+			[
+				SNew(SImage)
+				.ColorAndOpacity(FSlateColor::UseForeground())
+				.Image(MediaStreamStyle->GetBrush("MediaStreamEditor.OpenMedia.Small"))
+				.ToolTipText(LOCTEXT("Open Media", "Opens the current media, if the source is valid."))
+			]
+		]
+
+		// Close button.
+		+ SHorizontalBox::Slot()
+		.VAlign(VAlign_Center)
+		.AutoWidth()
+		[
+			SNew(SButton)
+			.VAlign(VAlign_Center)
+			.Visibility(this, &SMediaStreamPlaybackControls::Close_GetVisibility)
+			.OnClicked(this, &SMediaStreamPlaybackControls::Close_OnClicked)
+			[
+				SNew(SImage)
+				.ColorAndOpacity(FSlateColor::UseForeground())
+				.Image(MediaStreamStyle->GetBrush("MediaStreamEditor.CloseMedia.Small"))
+				.ToolTipText(LOCTEXT("Close Media", "Closes the currently opened media."))
+			]
+		]
+
+		// Rewind button.
+		+ SHorizontalBox::Slot()
+		.VAlign(VAlign_Center)
+		.AutoWidth()
+		[
+			SNew(SButton)
+			.VAlign(VAlign_Center)
+			.IsEnabled(this, &SMediaStreamPlaybackControls::Rewind_IsEnabled)
+			.OnClicked(this, &SMediaStreamPlaybackControls::Rewind_OnClicked)
+			[
+				SNew(SImage)
+				.ColorAndOpacity(FSlateColor::UseForeground())
+				.Image(MediaStreamStyle->GetBrush("MediaStreamEditor.RewindMedia.Small"))
+				.ToolTipText(LOCTEXT("Rewind", "Rewind the media to the beginning"))
+			]
+		]
+
+		// Reverse button.
+		+ SHorizontalBox::Slot()
+		.VAlign(VAlign_Center)
+		.AutoWidth()
+		[
+			SNew(SButton)
+			.VAlign(VAlign_Center)
+			.IsEnabled(this, &SMediaStreamPlaybackControls::Reverse_IsEnabled)
+			.OnClicked(this, &SMediaStreamPlaybackControls::Reverse_OnClicked)
+			[
+				SNew(SImage)
+				.ColorAndOpacity(FSlateColor::UseForeground())
+				.Image(MediaStreamStyle->GetBrush("MediaStreamEditor.ReverseMedia.Small"))
+				.ToolTipText(LOCTEXT("Reverse", "Reverse media playback"))
+			]
+		]
+
+		// Play button.
+		+ SHorizontalBox::Slot()
+		.VAlign(VAlign_Center)
+		.AutoWidth()
+		[
+			SNew(SButton)
+			.VAlign(VAlign_Center)
+			.Visibility(this, &SMediaStreamPlaybackControls::Play_GetVisibility)
+			.IsEnabled(this, &SMediaStreamPlaybackControls::Play_IsEnabled)
+			.OnClicked(this, &SMediaStreamPlaybackControls::Play_OnClicked)
+			[
+				SNew(SImage)
+				.ColorAndOpacity(FSlateColor::UseForeground())
+				.Image(MediaStreamStyle->GetBrush("MediaStreamEditor.PlayMedia.Small"))
+				.ToolTipText(LOCTEXT("Play", "Start media playback"))
+			]
+		]
+
+		// Pause button.
+		+ SHorizontalBox::Slot()
+		.VAlign(VAlign_Center)
+		.AutoWidth()
+		[
+			SNew(SButton)
+			.VAlign(VAlign_Center)
+			.Visibility(this, &SMediaStreamPlaybackControls::Pause_GetVisibility)
+			.IsEnabled(this, &SMediaStreamPlaybackControls::Pause_IsEnabled)
+			.OnClicked(this, &SMediaStreamPlaybackControls::Pause_OnClicked)
+			[
+				SNew(SImage)
+				.ColorAndOpacity(FSlateColor::UseForeground())
+				.Image(MediaStreamStyle->GetBrush("MediaStreamEditor.PauseMedia.Small"))
+				.ToolTipText(LOCTEXT("Pause", "Pause media playback"))
+			]
+		]
+
+		// Forward button.
+		+ SHorizontalBox::Slot()
+		.VAlign(VAlign_Center)
+		.AutoWidth()
+		[
+			SNew(SButton)
+			.VAlign(VAlign_Center)
+			.IsEnabled(this, &SMediaStreamPlaybackControls::Forward_IsEnabled)
+			.OnClicked(this, &SMediaStreamPlaybackControls::Forward_OnClicked)
+			[
+				SNew(SImage)
+				.ColorAndOpacity(FSlateColor::UseForeground())
+				.Image(MediaStreamStyle->GetBrush("MediaStreamEditor.ForwardMedia.Small"))
+				.RenderTransform(FSlateRenderTransform(TTransform2<float>(TQuat2<float>(PI), FVector2f(20.f, 20.f))))
+				.ToolTipText(LOCTEXT("Forward", "Fast forward media playback"))
+			]
+		]
+
+		// Add Media Track
+		+ SHorizontalBox::Slot()
+		.VAlign(VAlign_Center)
+		.AutoWidth()
+		[
+			SNew(SButton)
+			.ContentPadding(FMargin(0.f, 4.f, 0.f, 3.f))
+			.VAlign(VAlign_Center)
+			.Visibility(this, &SMediaStreamPlaybackControls::AddTrack_GetVisibility)
+			.OnClicked(this, &SMediaStreamPlaybackControls::AddTrack_OnClicked)
+			[
+				SNew(STextBlock)
+				.Font(IDetailLayoutBuilder::GetDetailFont())
+				.Text(LOCTEXT("AddTrack", "Add Track"))
+			]
+		]
+	];
+}
+
+UMediaStream* SMediaStreamPlaybackControls::GetMediaStream() const
+{
+	return MediaStreamWeak.Get();
+}
+
+UMediaPlayer* SMediaStreamPlaybackControls::GetMediaPlayer() const
+{
+	UMediaStream* MediaStream = MediaStreamWeak.Get();
+
+	if (!MediaStream)
+	{
+		return nullptr;
+	}
+
+	IMediaStreamPlayer* MediaStreamPlayer = MediaStream->GetPlayer().GetInterface();
+
+	if (!MediaStreamPlayer)
+	{
+		return nullptr;
+	}
+
+	return MediaStreamPlayer->GetPlayer();
+}
+
+void SMediaStreamPlaybackControls::OnChangePlaybackState(EMediaStreamPlaybackState InState)
+{
+	if (UMediaStream* MediaStream = GetMediaStream())
+	{
+		if (IMediaStreamPlayer* MediaStreamPlayer = MediaStream->GetPlayer().GetInterface())
+		{
+			if (MediaStreamPlayer->IsReadOnly())
+			{
+				return;
+			}
+
+			switch (InState)
+			{
+				case EMediaStreamPlaybackState::Play:
+				{
+					const float RequestedRate = FMath::Abs(MediaStreamPlayer->GetPlayerConfig().Rate);
+
+					if (RequestedRate != MediaStreamPlayer->GetPlayerConfig().Rate)
+					{
+						FMediaStreamPlayerConfig PlayerConfig = MediaStreamPlayer->GetPlayerConfig();
+						PlayerConfig.Rate = RequestedRate;
+
+						MediaStreamPlayer->SetPlayerConfig(PlayerConfig);
+					}
+
+					MediaStreamPlayer->Play();
+					break;
+				}
+
+				case EMediaStreamPlaybackState::Pause:
+					MediaStreamPlayer->Pause();
+					break;
+			}
+		}
+	}
+}
+
+void SMediaStreamPlaybackControls::OnChangePlaybackDirection(EMediaStreamPlaybackDirection InDirection)
+{
+	if (UMediaStream* MediaStream = GetMediaStream())
+	{
+		if (IMediaStreamPlayer* MediaStreamPlayer = MediaStream->GetPlayer().GetInterface())
+		{
+			if (MediaStreamPlayer->IsReadOnly())
+			{
+				return;
+			}
+
+			const float RequestedRate = FMath::Abs(MediaStreamPlayer->GetPlayerConfig().Rate)
+				* (InDirection == EMediaStreamPlaybackDirection::Forward ? 1.f : -1.f);
+
+			if (RequestedRate != MediaStreamPlayer->GetPlayerConfig().Rate)
+			{
+				FMediaStreamPlayerConfig PlayerConfig = MediaStreamPlayer->GetPlayerConfig();
+				PlayerConfig.Rate = RequestedRate;
+
+				MediaStreamPlayer->SetPlayerConfig(PlayerConfig);
+			}
+		}
+	}
+}
+
+void SMediaStreamPlaybackControls::OnPlaybackSeek(EMediaStreamPlaybackSeek InPosition)
+{
+	if (UMediaStream* MediaStream = GetMediaStream())
+	{
+		if (IMediaStreamPlayer* MediaStreamPlayer = MediaStream->GetPlayer().GetInterface())
+		{
+			if (MediaStreamPlayer->IsReadOnly())
+			{
+				return;
+			}
+
+			switch (InPosition)
+			{
+				case EMediaStreamPlaybackSeek::Previous:
+					MediaStreamPlayer->Previous();
+					break;
+
+				case EMediaStreamPlaybackSeek::Start:
+					MediaStreamPlayer->Rewind();
+					break;
+
+				case EMediaStreamPlaybackSeek::End:
+					MediaStreamPlayer->FastForward();
+					break;
+
+				case EMediaStreamPlaybackSeek::Next:
+					MediaStreamPlayer->Next();
+					break;
+			}
+		}
+	}
+}
+
+bool SMediaStreamPlaybackControls::IsRateSupported(float InRate) const
+{
+	UMediaPlayer* MediaPlayer = GetMediaPlayer();
+
+	return MediaPlayer && MediaPlayer->SupportsRate(InRate, /* Unthinned */ true);
+}
+
+EVisibility SMediaStreamPlaybackControls::Open_GetVisibility() const
+{
+	if (UMediaStream* MediaStream = GetMediaStream())
+	{
+		if (MediaStream->HasValidSource())
+		{
+			if (IMediaStreamPlayer* MediaStreamPlayer = MediaStream->GetPlayer().GetInterface())
+			{
+				if (UMediaPlayer* MediaPlayer = MediaStreamPlayer->GetPlayer())
+				{
+					if (MediaPlayer->IsClosed())
+					{
+						return EVisibility::Visible;
+					}
+
+					return EVisibility::Collapsed;
+				}
+			}
+
+			return EVisibility::Visible;
+		}
+	}
+
+	return EVisibility::Collapsed;
+}
+
+FReply SMediaStreamPlaybackControls::Open_OnClicked()
+{
+	if (UMediaStream* MediaStream = GetMediaStream())
+	{
+		if (IMediaStreamPlayer* MediaStreamPlayer = MediaStream->GetPlayer().GetInterface())
+		{
+			MediaStreamPlayer->OpenSource();
+		}
+	}
+
+	return FReply::Handled();
+}
+
+EVisibility SMediaStreamPlaybackControls::Close_GetVisibility() const
+{
+	if (UMediaStream* MediaStream = GetMediaStream())
+	{
+		if (MediaStream->HasValidSource())
+		{
+			if (IMediaStreamPlayer* MediaStreamPlayer = MediaStream->GetPlayer().GetInterface())
+			{
+				if (UMediaPlayer* MediaPlayer = MediaStreamPlayer->GetPlayer())
+				{
+					if (!FMediaStreamEditorSequencerLibrary::HasTrack(MediaStream))
+					{
+						if (!MediaPlayer->IsClosed())
+						{
+							return EVisibility::Visible;
+						}
+					}
+				}
+			}
+		}
+	}
+
+	return EVisibility::Collapsed;
+}
+
+FReply SMediaStreamPlaybackControls::Close_OnClicked()
+{
+	if (UMediaStream* MediaStream = GetMediaStream())
+	{
+		if (IMediaStreamPlayer* MediaStreamPlayer = MediaStream->GetPlayer().GetInterface())
+		{
+			MediaStreamPlayer->Close();
+		}
+	}
+
+	return FReply::Handled();
+}
+
+bool SMediaStreamPlaybackControls::Rewind_IsEnabled() const
+{
+	if (UMediaStream* MediaStream = GetMediaStream())
+	{
+		if (!FMediaStreamEditorSequencerLibrary::HasTrack(MediaStream))
+		{
+			if (const UMediaPlayer* MediaPlayer = GetMediaPlayer())
+			{
+				return MediaPlayer->IsReady() &&
+					MediaPlayer->SupportsSeeking() &&
+					MediaPlayer->GetTime() > FTimespan::Zero();
+			}
+		}
+	}
+
+	return false;
+}
+
+FReply SMediaStreamPlaybackControls::Rewind_OnClicked()
+{
+	OnPlaybackSeek(EMediaStreamPlaybackSeek::Previous);
+	return FReply::Handled();
+}
+
+bool SMediaStreamPlaybackControls::Reverse_IsEnabled() const
+{
+	if (UMediaStream* MediaStream = GetMediaStream())
+	{
+		if (!FMediaStreamEditorSequencerLibrary::HasTrack(MediaStream))
+		{
+			if (UMediaPlayer* MediaPlayer = GetMediaPlayer())
+			{
+				const float Rate = FMath::Abs(MediaPlayer->GetRate());
+
+				return MediaPlayer->IsReady() && IsRateSupported(Rate > 0 ? -Rate : -1.f);
+			}
+		}
+	}
+
+	return false;
+}
+
+FReply SMediaStreamPlaybackControls::Reverse_OnClicked()
+{
+	OnChangePlaybackDirection(EMediaStreamPlaybackDirection::Backward);
+	return FReply::Handled();
+}
+
+EVisibility SMediaStreamPlaybackControls::Play_GetVisibility() const
+{
+	if (const UMediaPlayer* MediaPlayer = GetMediaPlayer())
+	{
+		return MediaPlayer->IsPlaying() ? EVisibility::Collapsed : EVisibility::Visible;
+	}
+
+	return EVisibility::Visible;
+}
+
+bool SMediaStreamPlaybackControls::Play_IsEnabled() const
+{
+	if (UMediaStream* MediaStream = GetMediaStream())
+	{
+		if (!FMediaStreamEditorSequencerLibrary::HasTrack(MediaStream))
+		{
+			// Is the player paused or fast forwarding/rewinding?
+			if (const UMediaPlayer* MediaPlayer = GetMediaPlayer())
+			{
+				return MediaPlayer->IsReady()
+					&& (!MediaPlayer->IsPlaying() || (MediaPlayer->GetRate() != 1.0f));
+			}
+		}
+	}
+
+	return false; 
+}
+
+FReply SMediaStreamPlaybackControls::Play_OnClicked()
+{
+	OnChangePlaybackState(EMediaStreamPlaybackState::Play);
+	return FReply::Handled();
+}
+
+EVisibility SMediaStreamPlaybackControls::Pause_GetVisibility() const
+{
+	if (const UMediaPlayer* MediaPlayer = GetMediaPlayer())
+	{
+		const bool bIsVisible = MediaPlayer->CanPause() && !MediaPlayer->IsPaused();
+		return bIsVisible ? EVisibility::Visible : EVisibility::Collapsed;
+	}
+
+	return EVisibility::Collapsed;
+}
+
+bool SMediaStreamPlaybackControls::Pause_IsEnabled() const
+{
+	if (UMediaStream* MediaStream = GetMediaStream())
+	{
+		if (!FMediaStreamEditorSequencerLibrary::HasTrack(MediaStream))
+		{
+			if (const UMediaPlayer* MediaPlayer = GetMediaPlayer())
+			{
+				return MediaPlayer->CanPause() && !MediaPlayer->IsPaused();
+			}
+		}
+	}
+
+	return false;
+}
+
+FReply SMediaStreamPlaybackControls::Pause_OnClicked()
+{
+	OnChangePlaybackState(EMediaStreamPlaybackState::Pause);
+	return FReply::Handled();
+}
+
+bool SMediaStreamPlaybackControls::Forward_IsEnabled() const
+{
+	if (UMediaStream* MediaStream = GetMediaStream())
+	{
+		if (!FMediaStreamEditorSequencerLibrary::HasTrack(MediaStream))
+		{
+			if (UMediaPlayer* MediaPlayer = GetMediaPlayer())
+			{
+				return MediaPlayer->IsReady();
+			}
+		}
+	}
+
+	return false;
+}
+
+FReply SMediaStreamPlaybackControls::Forward_OnClicked()
+{
+	OnPlaybackSeek(EMediaStreamPlaybackSeek::End);
+	return FReply::Handled();
+}
+
+EVisibility SMediaStreamPlaybackControls::AddTrack_GetVisibility() const
+{
+	if (UMediaStream* MediaStream = GetMediaStream())
+	{
+		if (FMediaStreamEditorSequencerLibrary::CanAddTrack(MediaStream))
+		{
+			return EVisibility::Visible;
+		}
+	}
+
+	return EVisibility::Collapsed;
+}
+
+FReply SMediaStreamPlaybackControls::AddTrack_OnClicked()
+{
+	if (UMediaStream* MediaStream = GetMediaStream())
+	{
+		FMediaStreamEditorSequencerLibrary::AddTrack(MediaStream);
+	}
+
+	return FReply::Unhandled();
+}
+
+} // UE::MediaStreamEditor
+
+#undef LOCTEXT_NAMESPACE
diff --git a/Engine/Plugins/Experimental/MediaStream/Source/MediaStreamEditor/Private/DetailsPanel/SMediaStreamPlaybackControls.h b/Engine/Plugins/Experimental/MediaStream/Source/MediaStreamEditor/Private/DetailsPanel/SMediaStreamPlaybackControls.h
new file mode 100644
--- /dev/null
+++ b/Engine/Plugins/Experimental/MediaStream/Source/MediaStreamEditor/Private/DetailsPanel/SMediaStreamPlaybackControls.h
@@ -0,0 +1,81 @@
+// Copyright Epic Games, Inc. All Rights Reserved.
+
+#pragma once
+
+#include "Widgets/SCompoundWidget.h"
+
+#include "UObject/WeakObjectPtr.h"
+
+class SBox;
+class UMediaPlayer;
+class UMediaStream;
+enum class EMediaStreamPlaybackDirection : uint8;
+enum class EMediaStreamPlaybackSeek : uint8;
+enum class EMediaStreamPlaybackState : uint8;
+
+namespace UE::MediaStreamEditor
+{
+
+/**
+ * Displays media playback controls, such as play, pause, etc.
+ */
+class SMediaStreamPlaybackControls : public SCompoundWidget
+{
+	SLATE_BEGIN_ARGS(SMediaStreamPlaybackControls)
+		{}
+	SLATE_END_ARGS()
+
+	MEDIASTREAMEDITOR_API void Construct(const FArguments& InArgs, UMediaStream* InMediaStream);
+
+private:
+	/** The media stream for these controls. */
+	TWeakObjectPtr<UMediaStream> MediaStreamWeak;
+
+	/** Returns the media stream associated with these controls. */
+	UMediaStream* GetMediaStream() const;
+
+	/** Returns the media player associated with these controls. */
+	UMediaPlayer* GetMediaPlayer() const;
+
+	/** Changes the state of selected media streams. */
+	void OnChangePlaybackState(EMediaStreamPlaybackState InState);
+
+	/** Changes the state of selected media streams. */
+	void OnChangePlaybackDirection(EMediaStreamPlaybackDirection InDirection);
+
+	/** Changes the state of selected media streams. */
+	void OnPlaybackSeek(EMediaStreamPlaybackSeek InPosition);
+
+	/** Returns true if the given rate is supported by the media/player. */
+	bool IsRateSupported(float InRate) const;
+
+	/** Button delegates. */
+
+	EVisibility Open_GetVisibility() const;
+	FReply Open_OnClicked();
+
+	EVisibility Close_GetVisibility() const;
+	FReply Close_OnClicked();
+
+	bool Rewind_IsEnabled() const;
+	FReply Rewind_OnClicked();
+
+	bool Reverse_IsEnabled() const;
+	FReply Reverse_OnClicked();
+
+	EVisibility Play_GetVisibility() const;
+	bool Play_IsEnabled() const;
+	FReply Play_OnClicked();
+
+	EVisibility Pause_GetVisibility() const;
+	bool Pause_IsEnabled() const;
+	FReply Pause_OnClicked();
+
+	bool Forward_IsEnabled() const;
+	FReply Forward_OnClicked();
+
+	EVisibility AddTrack_GetVisibility() const;
+	FReply AddTrack_OnClicked();
+};
+
+} // UE::MediaStreamEditor
diff --git a/Engine/Plugins/Experimental/MediaStream/Source/MediaStreamEditor/Private/MediaStreamEditorModule.cpp b/Engine/Plugins/Experimental/MediaStream/Source/MediaStreamEditor/Private/MediaStreamEditorModule.cpp
--- a/Engine/Plugins/Experimental/MediaStream/Source/MediaStreamEditor/Private/MediaStreamEditorModule.cpp
+++ b/Engine/Plugins/Experimental/MediaStream/Source/MediaStreamEditor/Private/MediaStreamEditorModule.cpp
@@ -3,7 +3,6 @@
 #include "MediaStreamEditorModule.h"
 
 #include "DetailsPanel/MediaStreamCustomization.h"
-#include "DetailsPanel/MediaStreamSourceCustomization.h"
 #include "MediaStream.h"
 #include "MediaStreamEditorStyle.h"
 #include "MediaStreamEditorToolbar.h"
@@ -18,10 +17,7 @@ void FMediaStreamEditorModule::StartupModule()
 	FPropertyEditorModule& PropertyModule = FModuleManager::LoadModuleChecked<FPropertyEditorModule>("PropertyEditor");
 
 	PropertyModule.RegisterCustomClassLayout(UMediaStream::StaticClass()->GetFName(),
-		FOnGetDetailCustomizationInstance::CreateStatic(&FMediaStreamCustomization::MakeInstance));
-
-	PropertyModule.RegisterCustomPropertyTypeLayout(FMediaStreamSource::StaticStruct()->GetFName(),
-		FOnGetPropertyTypeCustomizationInstance::CreateStatic(&FMediaStreamSourceCustomization::MakeInstance));
+		FOnGetDetailCustomizationInstance::CreateStatic(&UE::MediaStreamEditor::FMediaStreamCustomization::MakeInstance));
 
 	FMediaStreamEditorToolbar::Get().Integrate();
 }
diff --git a/Engine/Plugins/Experimental/MediaStream/Source/MediaStreamEditor/Private/MediaStreamWidgets.cpp b/Engine/Plugins/Experimental/MediaStream/Source/MediaStreamEditor/Private/MediaStreamWidgets.cpp
new file mode 100644
--- /dev/null
+++ b/Engine/Plugins/Experimental/MediaStream/Source/MediaStreamEditor/Private/MediaStreamWidgets.cpp
@@ -0,0 +1,34 @@
+// Copyright Epic Games, Inc. All Rights Reserved.
+
+#include "MediaStreamWidgets.h"
+
+#include "DetailsPanel/MediaStreamSourceCustomization.h"
+#include "DetailsPanel/SMediaStreamMediaDetails.h"
+#include "DetailsPanel/SMediaStreamMediaTrack.h"
+#include "DetailsPanel/SMediaStreamPlaybackControls.h"
+#include "Templates/SharedPointer.h"
+
+namespace UE::MediaStreamEditor
+{
+
+IMediaStreamSchemeHandler::FCustomWidgets FMediaStreamWidgets::GenerateSourceSchemeRows(UMediaStream* InMediaStream)
+{
+	return FMediaStreamSourceCustomization().Customize(InMediaStream);
+}
+
+TSharedRef<SWidget> FMediaStreamWidgets::CreateTextureDetailsWidget(UMediaStream* InMediaStream)
+{
+	return SNew(SMediaStreamMediaDetails, InMediaStream);
+}
+
+TSharedRef<SWidget> FMediaStreamWidgets::CreateTrackWidget(UMediaStream* InMediaStream)
+{
+	return SNew(SMediaStreamMediaTrack, InMediaStream);
+}
+
+TSharedRef<SWidget> FMediaStreamWidgets::CreateControlsWidget(UMediaStream* InMediaStream)
+{
+	return SNew(SMediaStreamPlaybackControls, InMediaStream);
+}
+
+}; // UE::MediaStreamEditor
diff --git a/Engine/Plugins/Experimental/MediaStream/Source/MediaStreamEditor/Public/MediaStreamWidgets.h b/Engine/Plugins/Experimental/MediaStream/Source/MediaStreamEditor/Public/MediaStreamWidgets.h
new file mode 100644
--- /dev/null
+++ b/Engine/Plugins/Experimental/MediaStream/Source/MediaStreamEditor/Public/MediaStreamWidgets.h
@@ -0,0 +1,38 @@
+// Copyright Epic Games, Inc. All Rights Reserved.
+
+#pragma once
+
+#include "IMediaStreamSchemeHandler.h"
+#include "Templates/SharedPointerFwd.h"
+
+class SWidget;
+class UMediaStream;
+
+namespace UE::MediaStreamEditor
+{
+
+class FMediaStreamWidgets
+{
+public:
+	/**
+	 * Generates custom widgets for source scheme types.
+	 */
+	MEDIASTREAMEDITOR_API static IMediaStreamSchemeHandler::FCustomWidgets GenerateSourceSchemeRows(UMediaStream* InMediaStream);
+
+	/**
+	 * Creates a widget representing the details of the owned Media Texture.
+	 */
+	MEDIASTREAMEDITOR_API static TSharedRef<SWidget> CreateTextureDetailsWidget(UMediaStream* InMediaStream);
+
+	/**
+	 * Creates a widget allowing the user to scrub the media.
+	 */
+	MEDIASTREAMEDITOR_API static TSharedRef<SWidget> CreateTrackWidget(UMediaStream* InMediaStream);
+
+	/**
+	 * Creates a widget allowing the user to control the media playback.
+	 */
+	MEDIASTREAMEDITOR_API static TSharedRef<SWidget> CreateControlsWidget(UMediaStream* InMediaStream);
+};
+
+} // UE::MediaStreamEditor
