Material Designer: Media source selection
- The panel now refreshes when source details are changed.

diff --git a/Engine/Plugins/Experimental/DynamicMaterialMediaStreamBridge/Source/DynamicMaterialMediaStreamBridge/Private/DMMaterialValueMediaStream.cpp b/Engine/Plugins/Experimental/DynamicMaterialMediaStreamBridge/Source/DynamicMaterialMediaStreamBridge/Private/DMMaterialValueMediaStream.cpp
--- a/Engine/Plugins/Experimental/DynamicMaterialMediaStreamBridge/Source/DynamicMaterialMediaStreamBridge/Private/DMMaterialValueMediaStream.cpp
+++ b/Engine/Plugins/Experimental/DynamicMaterialMediaStreamBridge/Source/DynamicMaterialMediaStreamBridge/Private/DMMaterialValueMediaStream.cpp
@@ -184,6 +184,11 @@ void UDMMaterialValueMediaStream::PostEditChangeProperty(FPropertyChangedEvent&
 {
 	Super::PostEditChangeProperty(InPropertyChangedEvent);
 
+	if (InPropertyChangedEvent.GetMemberPropertyName() == UMediaStream::GetSourcePropertyName())
+	{
+		Update(this, EDMUpdateType::RefreshDetailView);
+	}
+
 	UpdatePlayer();
 }
 
diff --git a/Engine/Plugins/Experimental/MediaStream/Source/MediaStream/Private/MediaStreamSchemeHandlerLibrary.cpp b/Engine/Plugins/Experimental/MediaStream/Source/MediaStream/Private/MediaStreamSchemeHandlerLibrary.cpp
--- a/Engine/Plugins/Experimental/MediaStream/Source/MediaStream/Private/MediaStreamSchemeHandlerLibrary.cpp
+++ b/Engine/Plugins/Experimental/MediaStream/Source/MediaStream/Private/MediaStreamSchemeHandlerLibrary.cpp
@@ -139,14 +139,25 @@ void FMediaStreamSchemeHandlerLibrary::SetSource(UMediaStream* InMediaStream, co
 		return;
 	}
 
-	FProperty* Property = InMediaStream->GetClass()->FindPropertyByName( UMediaStream::GetSourcePropertyName() );
+	FProperty* Property = InMediaStream->GetClass()->FindPropertyByName(UMediaStream::GetSourcePropertyName());
+	UObject* Outer = InMediaStream->GetOuter();
 
 	InMediaStream->PreEditChange(Property);
 
+	if (Outer)
+	{
+		Outer->PreEditChange(Property);
+	}
+
 	InMediaStream->SetSource(FMediaStreamSchemeHandlerManager::Get().CreateSource(InMediaStream, InScheme, InPath));
 
 	FPropertyChangedEvent Event(Property, EPropertyChangeType::Interactive, {InMediaStream});
 	InMediaStream->PostEditChangeProperty(Event);
+
+	if (Outer)
+	{
+		Outer->PostEditChangeProperty(Event);
+	}
 }
 
 void FMediaStreamSchemeHandlerLibrary::SetSource(TWeakObjectPtr<UMediaStream> InMediaStreamWeak, const FName& InScheme, const FString& InPath)
