Media Stream: Better auto-load/play conditions
- Now only blocks auto-load/play in the editor when a map is loading.

diff --git a/Engine/Plugins/Experimental/MediaStream/Source/MediaStream/Private/MediaStreamModule.cpp b/Engine/Plugins/Experimental/MediaStream/Source/MediaStream/Private/MediaStreamModule.cpp
--- a/Engine/Plugins/Experimental/MediaStream/Source/MediaStream/Private/MediaStreamModule.cpp
+++ b/Engine/Plugins/Experimental/MediaStream/Source/MediaStream/Private/MediaStreamModule.cpp
@@ -14,6 +14,10 @@
 #include "SchemeHandlers/MediaStreamSubobjectSchemeHandler.h"
 #include "UObject/UObjectBase.h"
 
+#if WITH_EDITOR
+#include "Editor.h"
+#endif
+
 DEFINE_LOG_CATEGORY(LogMediaStream);
 
 FMediaStreamModule& FMediaStreamModule::Get()
@@ -24,7 +28,7 @@ FMediaStreamModule& FMediaStreamModule::Get()
 bool FMediaStreamModule::CanOpenSourceOnLoad()
 {
 #if WITH_EDITOR
-	return !GIsEditor;
+	return !GIsEditor || !bIsMapLoading;
 #else
 	return true;
 #endif
@@ -32,7 +36,11 @@ bool FMediaStreamModule::CanOpenSourceOnLoad()
 
 bool FMediaStreamModule::CanAutoplay() const
 {
-	return !GIsEditor;
+#if WITH_EDITOR
+	return !GIsEditor || !bIsMapLoading;
+#else
+	return true;
+#endif	
 }
 
 void FMediaStreamModule::StartupModule()
@@ -47,6 +55,13 @@ void FMediaStreamModule::StartupModule()
 	ObjectHandlerManager.RegisterObjectHandler<FMediaStreamMediaPlaylistHandler>();
 	ObjectHandlerManager.RegisterObjectHandler<FMediaStreamMediaSourceHandler>();
 	ObjectHandlerManager.RegisterObjectHandler<FMediaStreamMediaStreamHandler>();
+
+#if WITH_EDITOR
+	bIsMapLoading = false;
+
+	FEditorDelegates::OnMapLoad.AddRaw(this, &FMediaStreamModule::OnMapLoad);
+	FEditorDelegates::OnMapOpened.AddRaw(this, &FMediaStreamModule::OnMapOpened);
+#endif
 }
 
 void FMediaStreamModule::ShutdownModule()
@@ -66,6 +81,27 @@ void FMediaStreamModule::ShutdownModule()
 	ObjectHandlerManager.UnregisterObjectHandler<FMediaStreamMediaPlaylistHandler>();
 	ObjectHandlerManager.UnregisterObjectHandler<FMediaStreamMediaSourceHandler>();
 	ObjectHandlerManager.UnregisterObjectHandler<FMediaStreamMediaStreamHandler>();
+
+#if WITH_EDITOR
+	FEditorDelegates::OnMapLoad.RemoveAll(this);
+	FEditorDelegates::OnMapOpened.RemoveAll(this);
+#endif
 }
 
+#if WITH_EDITOR
+void FMediaStreamModule::OnMapLoad(const FString& InFilename, FCanLoadMap& OutCanLoadMap)
+{
+	// If this has been set to false, don't set the flag.
+	if (OutCanLoadMap.Get())
+	{
+		bIsMapLoading = true;
+	}
+}
+
+void FMediaStreamModule::OnMapOpened(const FString& InFilename, bool bInTemplate)
+{
+	bIsMapLoading = false;
+}
+#endif
+
 IMPLEMENT_MODULE(FMediaStreamModule, MediaStream);
diff --git a/Engine/Plugins/Experimental/MediaStream/Source/MediaStream/Public/MediaStreamModule.h b/Engine/Plugins/Experimental/MediaStream/Source/MediaStream/Public/MediaStreamModule.h
--- a/Engine/Plugins/Experimental/MediaStream/Source/MediaStream/Public/MediaStreamModule.h
+++ b/Engine/Plugins/Experimental/MediaStream/Source/MediaStream/Public/MediaStreamModule.h
@@ -4,12 +4,14 @@
 
 #include "Modules/ModuleInterface.h"
 
+#include "Containers/UnrealString.h"
 #include "Delegates/Delegate.h"
 #include "Delegates/DelegateCombinations.h"
 #include "Logging/LogMacros.h"
 
 class IMediaStreamPlayer;
 enum class EMediaStreamPlaybackState : uint8;
+struct FCanLoadMap;
 
 DECLARE_LOG_CATEGORY_EXTERN(LogMediaStream, Log, All);
 
@@ -29,4 +31,13 @@ public:
 	virtual void StartupModule() override;
 	virtual void ShutdownModule() override;
 	//~ End IModuleInterface
+
+#if WITH_EDITOR
+protected:
+	bool bIsMapLoading;
+
+	void OnMapLoad(const FString& InFilename, FCanLoadMap& OutCanLoadMap);
+
+	void OnMapOpened(const FString& InFilename, bool bInTemplate);
+#endif
 };
