Media Stream / Media Viewer: Closing properly
- The media stream now has a Close method to close the active player.
- The media source image viewer now closes its media stream properly when it's destroyed.

diff --git a/Engine/Plugins/Experimental/MediaStream/Source/MediaStream/Private/MediaStream.cpp b/Engine/Plugins/Experimental/MediaStream/Source/MediaStream/Private/MediaStream.cpp
--- a/Engine/Plugins/Experimental/MediaStream/Source/MediaStream/Private/MediaStream.cpp
+++ b/Engine/Plugins/Experimental/MediaStream/Source/MediaStream/Private/MediaStream.cpp
@@ -92,6 +92,11 @@ UMediaStream::FOnSourceChanged& UMediaStream::GetOnSourceChanged()
 	return OnSourceChanged;
 }
 
+void UMediaStream::Close()
+{
+	SetSource(FMediaStreamSource());
+}
+
 #if WITH_EDITOR
 void UMediaStream::PreEditUndo()
 {
@@ -146,12 +151,13 @@ void UMediaStream::BeginDestroy()
 {
 	Super::BeginDestroy();
 
+	Source = FMediaStreamSource();
+
 	if (IsValid(PlayerObject))
 	{
 		if (IMediaStreamPlayer* Player = Cast<IMediaStreamPlayer>(PlayerObject))
 		{
 			Player->Deinitialize();
-			PlayerObject = nullptr;
 		}
 	}
 }
diff --git a/Engine/Plugins/Experimental/MediaStream/Source/MediaStream/Private/Players/MediaStreamLocalPlayer.cpp b/Engine/Plugins/Experimental/MediaStream/Source/MediaStream/Private/Players/MediaStreamLocalPlayer.cpp
--- a/Engine/Plugins/Experimental/MediaStream/Source/MediaStream/Private/Players/MediaStreamLocalPlayer.cpp
+++ b/Engine/Plugins/Experimental/MediaStream/Source/MediaStream/Private/Players/MediaStreamLocalPlayer.cpp
@@ -517,6 +517,16 @@ void UMediaStreamLocalPlayer::InitTexture()
 
 void UMediaStreamLocalPlayer::DeinitTexture()
 {
+	if (IsValid(MediaTexture))
+	{
+		const bool bIsPlayerValid = IsValid(Player);
+
+		MediaTexture->SetMediaPlayer(bIsPlayerValid ? Player : nullptr);
+
+#if WITH_EDITOR
+		MediaTexture->SetDefaultMediaPlayer(bIsPlayerValid ? Player : nullptr);
+#endif
+	}
 }
 
 void UMediaStreamLocalPlayer::InitPlayer()
diff --git a/Engine/Plugins/Experimental/MediaStream/Source/MediaStream/Public/MediaStream.h b/Engine/Plugins/Experimental/MediaStream/Source/MediaStream/Public/MediaStream.h
--- a/Engine/Plugins/Experimental/MediaStream/Source/MediaStream/Public/MediaStream.h
+++ b/Engine/Plugins/Experimental/MediaStream/Source/MediaStream/Public/MediaStream.h
@@ -74,6 +74,11 @@ public:
 	DECLARE_DYNAMIC_MULTICAST_DELEGATE_OneParam(FOnSourceChanged, UMediaStream*, InMediaStream);
 	MEDIASTREAM_API FOnSourceChanged& GetOnSourceChanged();
 
+	/**
+	 * Stops the media stream and unloads any resources.
+	 */
+	MEDIASTREAM_API void Close();
+
 	//~ Begin UObject
 #if WITH_EDITOR
 	virtual void PreEditUndo() override;
diff --git a/Engine/Plugins/Experimental/MediaViewer/Source/MediaViewer/Private/ImageViewers/MediaSourceImageViewer.cpp b/Engine/Plugins/Experimental/MediaViewer/Source/MediaViewer/Private/ImageViewers/MediaSourceImageViewer.cpp
--- a/Engine/Plugins/Experimental/MediaViewer/Source/MediaViewer/Private/ImageViewers/MediaSourceImageViewer.cpp
+++ b/Engine/Plugins/Experimental/MediaViewer/Source/MediaViewer/Private/ImageViewers/MediaSourceImageViewer.cpp
@@ -296,6 +296,14 @@ FMediaSourceImageViewer::FMediaSourceImageViewer(const FGuid& InId, TNotNull<UMe
 	}
 }
 
+FMediaSourceImageViewer::~FMediaSourceImageViewer()
+{
+	if (UObjectInitialized() && MediaStream)
+	{
+		MediaStream->Close();
+	}
+}
+
 UMediaStream* FMediaSourceImageViewer::GetMediaStream() const
 {
 	return MediaStream;
diff --git a/Engine/Plugins/Experimental/MediaViewer/Source/MediaViewer/Private/ImageViewers/MediaSourceImageViewer.h b/Engine/Plugins/Experimental/MediaViewer/Source/MediaViewer/Private/ImageViewers/MediaSourceImageViewer.h
--- a/Engine/Plugins/Experimental/MediaViewer/Source/MediaViewer/Private/ImageViewers/MediaSourceImageViewer.h
+++ b/Engine/Plugins/Experimental/MediaViewer/Source/MediaViewer/Private/ImageViewers/MediaSourceImageViewer.h
@@ -112,7 +112,7 @@ public:
 	FMediaSourceImageViewer(TNotNull<UMediaSource*> InMediaSource, const FText& InDisplayName);
 	FMediaSourceImageViewer(const FGuid& InId, TNotNull<UMediaSource*> InMediaSource, const FText& InDisplayName);
 
-	virtual ~FMediaSourceImageViewer() override = default;
+	virtual ~FMediaSourceImageViewer() override;
 
 	UMediaStream* GetMediaStream() const;
 
