Media Stream: Actor no longer placeable.
- The Media Stream Actor is no longer placeable.

diff --git a/Engine/Plugins/Experimental/MediaStream/Source/MediaStream/Public/MediaStreamActor.h b/Engine/Plugins/Experimental/MediaStream/Source/MediaStream/Public/MediaStreamActor.h
--- a/Engine/Plugins/Experimental/MediaStream/Source/MediaStream/Public/MediaStreamActor.h
+++ b/Engine/Plugins/Experimental/MediaStream/Source/MediaStream/Public/MediaStreamActor.h
@@ -8,7 +8,7 @@
 
 class UMediaStreamComponent;
 
-UCLASS(BlueprintType, MinimalAPI, ComponentWrapperClass,
+UCLASS(BlueprintType, MinimalAPI, ComponentWrapperClass, NotPlaceable,
 	PrioritizeCategories = (MediaStream, MediaControls, MediaSource, MediaDetails, MediaTexture, MediaCache, MediaPlayer))
 class AMediaStreamActor : public AActor
 {
