Material Designer: Duel playhead control
- The play controls in the Material Designer for media now (attempt) to control both the source and preview media.

diff --git a/Engine/Plugins/Experimental/DynamicMaterialMediaStreamBridge/Source/DynamicMaterialMediaStreamBridgeEditor/Private/DMMaterialValueMediaStreamPropertyRowGenerator.cpp b/Engine/Plugins/Experimental/DynamicMaterialMediaStreamBridge/Source/DynamicMaterialMediaStreamBridgeEditor/Private/DMMaterialValueMediaStreamPropertyRowGenerator.cpp
--- a/Engine/Plugins/Experimental/DynamicMaterialMediaStreamBridge/Source/DynamicMaterialMediaStreamBridgeEditor/Private/DMMaterialValueMediaStreamPropertyRowGenerator.cpp
+++ b/Engine/Plugins/Experimental/DynamicMaterialMediaStreamBridge/Source/DynamicMaterialMediaStreamBridgeEditor/Private/DMMaterialValueMediaStreamPropertyRowGenerator.cpp
@@ -57,7 +57,7 @@ void FDMMaterialValueMediaStreamPropertyRowGenerator::AddComponentProperties(FDM
 
 	if (bHasPlayer)
 	{
-		AddControlCategory(MediaStreamParams, /* Preview */ true);
+		AddControlCategory(MediaStreamParams);
 	}
 
 	AddSourceCategory(MediaStreamParams);
@@ -67,65 +67,51 @@ void FDMMaterialValueMediaStreamPropertyRowGenerator::AddComponentProperties(FDM
 		AddDetailsCategory(MediaStreamParams);
 		AddCacheCategory(MediaStreamParams);
 		AddPlayerProperty(MediaStreamParams);
-
-		AddControlCategory(MediaStreamParams, /* Preview */ false);
 	}
 }
 
-void FDMMaterialValueMediaStreamPropertyRowGenerator::AddControlCategory(FDMComponentPropertyRowGeneratorParams& InParams, 
-	bool bInPreview)
+void FDMMaterialValueMediaStreamPropertyRowGenerator::AddControlCategory(FDMComponentPropertyRowGeneratorParams& InParams)
 {
 	FDMComponentPropertyRowGeneratorParams* Params = &InParams;
-	UMediaStream* MediaStream = nullptr;
-	FDMComponentPropertyRowGeneratorParams OriginalParams = InParams;
-	
-	if (bInPreview)
-	{
-		MediaStream = Cast<UMediaStream>(InParams.Object);
-	}
-	else
-	{
-		const FString RelativePath = InParams.Object->GetPathName(InParams.PreviewMaterialModelBase);
-		MediaStream = UDMMaterialModelFunctionLibrary::FindSubobject<UMediaStream>(InParams.OriginalMaterialModelBase, RelativePath);
-		OriginalParams.Object = MediaStream;
-		Params = &OriginalParams;
-	}
+	UMediaStream* PreviewMediaStream = Cast<UMediaStream>(InParams.Object);
 
-	if (!MediaStream)
+	if (!PreviewMediaStream)
 	{
 		return;
 	}
+
+	const FString RelativePath = InParams.Object->GetPathName(InParams.PreviewMaterialModelBase);
+	UMediaStream* OriginalMediaStream = UDMMaterialModelFunctionLibrary::FindSubobject<UMediaStream>(InParams.OriginalMaterialModelBase, RelativePath);
+
+	TArray<UMediaStream*> MediaStreams = !!OriginalMediaStream
+		? TArray<UMediaStream*>({OriginalMediaStream, PreviewMediaStream})
+		: TArray<UMediaStream*>({PreviewMediaStream});
 	
-	const FName CategoryName = bInPreview
-		? "Preview Media Controls"
-		: "Source Media Controls";
+	const FName CategoryName = "Media Controls";
 
 	FDMPropertyHandle& TrackHandle = Params->PropertyRows->Add_GetRef(
 		IDMWidgetLibrary::Get().GetPropertyHandle(Params->CreatePropertyHandleParams(TEXT("Track")))
 	);
 	TrackHandle.CategoryOverrideName = CategoryName;
-	TrackHandle.ValueWidget = UE::MediaStreamEditor::FMediaStreamWidgets::CreateTrackWidget(MediaStream);
-	TrackHandle.ValueName = bInPreview ? "PreviewTrack" : "SourceTrack";
+	TrackHandle.ValueWidget = UE::MediaStreamEditor::FMediaStreamWidgets::CreateTrackWidget(MediaStreams);
+	TrackHandle.ValueName = "Track";
 
 	FDMPropertyHandle& ControlHandle = Params->PropertyRows->Add_GetRef(
 		IDMWidgetLibrary::Get().GetPropertyHandle(Params->CreatePropertyHandleParams(TEXT("Controls")))
 	);
 	ControlHandle.CategoryOverrideName = CategoryName;
-	ControlHandle.ValueWidget = UE::MediaStreamEditor::FMediaStreamWidgets::CreateControlsWidget(MediaStream);
-	ControlHandle.ValueName = bInPreview ? "PreviewControls" : "SourceControls";
+	ControlHandle.ValueWidget = UE::MediaStreamEditor::FMediaStreamWidgets::CreateControlsWidget(MediaStreams);
+	ControlHandle.ValueName = "Controls";
 
-	if (bInPreview)
-	{
-		FDMPropertyHandle& PlayOnOpenHandle = Params->PropertyRows->Add_GetRef(
-			IDMWidgetLibrary::Get().GetPropertyHandle(Params->CreatePropertyHandleParams(TEXT("bPlayOnOpen")))
-		);
-		PlayOnOpenHandle.CategoryOverrideName = CategoryName;
+	FDMPropertyHandle& PlayOnOpenHandle = Params->PropertyRows->Add_GetRef(
+		IDMWidgetLibrary::Get().GetPropertyHandle(Params->CreatePropertyHandleParams(TEXT("bPlayOnOpen")))
+	);
+	PlayOnOpenHandle.CategoryOverrideName = CategoryName;
 
-		FDMPropertyHandle& LoopingHandle = Params->PropertyRows->Add_GetRef(
-			IDMWidgetLibrary::Get().GetPropertyHandle(Params->CreatePropertyHandleParams(TEXT("bLooping")))
-		);
-		LoopingHandle.CategoryOverrideName = CategoryName;
-	}
+	FDMPropertyHandle& LoopingHandle = Params->PropertyRows->Add_GetRef(
+		IDMWidgetLibrary::Get().GetPropertyHandle(Params->CreatePropertyHandleParams(TEXT("bLooping")))
+	);
+	LoopingHandle.CategoryOverrideName = CategoryName;
 }
 
 void FDMMaterialValueMediaStreamPropertyRowGenerator::AddSourceCategory(FDMComponentPropertyRowGeneratorParams& InParams)
diff --git a/Engine/Plugins/Experimental/DynamicMaterialMediaStreamBridge/Source/DynamicMaterialMediaStreamBridgeEditor/Private/DMMaterialValueMediaStreamPropertyRowGenerator.h b/Engine/Plugins/Experimental/DynamicMaterialMediaStreamBridge/Source/DynamicMaterialMediaStreamBridgeEditor/Private/DMMaterialValueMediaStreamPropertyRowGenerator.h
--- a/Engine/Plugins/Experimental/DynamicMaterialMediaStreamBridge/Source/DynamicMaterialMediaStreamBridgeEditor/Private/DMMaterialValueMediaStreamPropertyRowGenerator.h
+++ b/Engine/Plugins/Experimental/DynamicMaterialMediaStreamBridge/Source/DynamicMaterialMediaStreamBridgeEditor/Private/DMMaterialValueMediaStreamPropertyRowGenerator.h
@@ -17,7 +17,7 @@ public:
 	virtual void AddComponentProperties(FDMComponentPropertyRowGeneratorParams& InParams) override;
 	//~ End FDMComponentPropertyRowGenerator
 
-	void AddControlCategory(FDMComponentPropertyRowGeneratorParams& InParams, bool bInPreview);
+	void AddControlCategory(FDMComponentPropertyRowGeneratorParams& InParams);
 
 	void AddSourceCategory(FDMComponentPropertyRowGeneratorParams& InParams);
 
diff --git a/Engine/Plugins/Experimental/MediaStream/Source/MediaStreamEditor/Private/DetailsPanel/MediaStreamCustomization.cpp b/Engine/Plugins/Experimental/MediaStream/Source/MediaStreamEditor/Private/DetailsPanel/MediaStreamCustomization.cpp
--- a/Engine/Plugins/Experimental/MediaStream/Source/MediaStreamEditor/Private/DetailsPanel/MediaStreamCustomization.cpp
+++ b/Engine/Plugins/Experimental/MediaStream/Source/MediaStreamEditor/Private/DetailsPanel/MediaStreamCustomization.cpp
@@ -93,13 +93,13 @@ void FMediaStreamCustomization::AddControlCategory(IDetailLayoutBuilder& InDetai
 				return false;
 			}))
 		[
-			FMediaStreamWidgets::CreateTrackWidget(MediaStream)
+			FMediaStreamWidgets::CreateTrackWidget({MediaStream})
 		];
 
 	// Add media control buttons.
 	ControlsCategory.AddCustomRow(LOCTEXT("Controls", "Controls"))
 	[
-		FMediaStreamWidgets::CreateControlsWidget(MediaStream)
+		FMediaStreamWidgets::CreateControlsWidget({MediaStream})
 	];
 
 	TSharedPtr<IPropertyHandle> PlayerHandle = InDetailBuilder.GetProperty(UMediaStream::GetPlayerPropertyName());
diff --git a/Engine/Plugins/Experimental/MediaStream/Source/MediaStreamEditor/Private/DetailsPanel/SMediaStreamMediaTrack.cpp b/Engine/Plugins/Experimental/MediaStream/Source/MediaStreamEditor/Private/DetailsPanel/SMediaStreamMediaTrack.cpp
--- a/Engine/Plugins/Experimental/MediaStream/Source/MediaStreamEditor/Private/DetailsPanel/SMediaStreamMediaTrack.cpp
+++ b/Engine/Plugins/Experimental/MediaStream/Source/MediaStreamEditor/Private/DetailsPanel/SMediaStreamMediaTrack.cpp
@@ -13,9 +13,14 @@
 namespace UE::MediaStreamEditor
 {
 
-void SMediaStreamMediaTrack::Construct(const FArguments& InArgs, UMediaStream* InMediaStream)
+void SMediaStreamMediaTrack::Construct(const FArguments& InArgs, const TArray<UMediaStream*>& InMediaStreams)
 {
-	MediaStreamWeak = InMediaStream;
+	MediaStreamsWeak.Reserve(InMediaStreams.Num());
+
+	for (UMediaStream* MediaStream : InMediaStreams)
+	{
+		MediaStreamsWeak.Add(MediaStream);
+	}
 
 	SetCanTick(true);
 
@@ -34,39 +39,58 @@ void SMediaStreamMediaTrack::Tick(const FGeometry& InAllottedGeometry, const dou
 {
 	SCompoundWidget::Tick(InAllottedGeometry, InCurrentTime, InDeltaTime);
 
-	UMediaPlayer* MediaPlayer = GetMediaPlayer();
-	UMediaPlayer* CurrentMediaPlayer = MediaPlayerWeak.Get();
+	TArray<UMediaPlayer*> MediaPlayers = GetMediaPlayers();
+	bool bNeedsUpdate = MediaPlayers.Num() != MediaPlayersWeak.Num();
 
-	if (MediaPlayer != CurrentMediaPlayer)
+	if (!bNeedsUpdate)
 	{
-		Content->SetContent(CreateTrack());
+		for (int32 Index = 0; Index < MediaPlayersWeak.Num(); ++Index)
+		{
+			UMediaPlayer* MediaPlayer = MediaPlayersWeak[Index].Get();
+
+			if (MediaPlayer != MediaPlayers[Index])
+			{
+				bNeedsUpdate = true;
+				break;
+			}
+		}
 	}
-}
-
-UMediaPlayer* SMediaStreamMediaTrack::GetMediaPlayer() const
-{
-	UMediaStream* MediaStream = MediaStreamWeak.Get();
 
-	if (!MediaStream)
+	if (!bNeedsUpdate)
 	{
-		return nullptr;
+		return;
 	}
 
-	IMediaStreamPlayer* MediaStreamPlayer = MediaStream->GetPlayer().GetInterface();
+	Content->SetContent(CreateTrack());
+}
+
+TArray<UMediaPlayer*> SMediaStreamMediaTrack::GetMediaPlayers() const
+{
+	TArray<UMediaPlayer*> MediaPlayers;
+	MediaPlayers.Reserve(MediaStreamsWeak.Num());
 
-	if (!MediaStreamPlayer)
+	for (const TWeakObjectPtr<UMediaStream>& MediaStreamWeak : MediaStreamsWeak)
 	{
-		return nullptr;
+		if (UMediaStream* MediaStream = MediaStreamWeak.Get())
+		{
+			if (IMediaStreamPlayer* MediaStreamPlayer = MediaStream->GetPlayer().GetInterface())
+			{
+				if (UMediaPlayer* MediaPlayer = MediaStreamPlayer->GetPlayer())
+				{
+					MediaPlayers.Add(MediaPlayer);
+				}
+			}
+		}
 	}
 
-	return MediaStreamPlayer->GetPlayer();
+	return MediaPlayers;
 }
 
 TSharedRef<SWidget> SMediaStreamMediaTrack::CreateTrack()
 {
-	UMediaPlayer* MediaPlayer = GetMediaPlayer();
+	TArray<UMediaPlayer*> MediaPlayers = GetMediaPlayers();
 
-	if (!MediaPlayer)
+	if (MediaPlayers.IsEmpty())
 	{
 		return SNullWidget::NullWidget;
 	}
@@ -78,12 +102,14 @@ TSharedRef<SWidget> SMediaStreamMediaTrack::CreateTrack()
 		return SNullWidget::NullWidget;
 	}
 
-	MediaPlayerWeak = MediaPlayer;
+	MediaPlayersWeak.Empty(MediaPlayers.Num());
 
-	TArray<TWeakObjectPtr<UMediaPlayer>> MediaPlayerList;
-	MediaPlayerList.Add(MediaPlayer);
+	for (UMediaPlayer* MediaPlayer : MediaPlayers)
+	{
+		MediaPlayersWeak.Add(MediaPlayer);
+	}
 
-	const TSharedRef<IMediaPlayerSlider> MediaPlayerSlider = MediaPlayerEditorModule->CreateMediaPlayerSliderWidget(MediaPlayerList);
+	const TSharedRef<IMediaPlayerSlider> MediaPlayerSlider = MediaPlayerEditorModule->CreateMediaPlayerSliderWidget(MediaPlayersWeak);
 
 	MediaPlayerSlider->SetSliderHandleColor(FSlateColor(EStyleColor::AccentBlue));
 	MediaPlayerSlider->SetVisibleWhenInactive(EVisibility::Visible);
diff --git a/Engine/Plugins/Experimental/MediaStream/Source/MediaStreamEditor/Private/DetailsPanel/SMediaStreamMediaTrack.h b/Engine/Plugins/Experimental/MediaStream/Source/MediaStreamEditor/Private/DetailsPanel/SMediaStreamMediaTrack.h
--- a/Engine/Plugins/Experimental/MediaStream/Source/MediaStreamEditor/Private/DetailsPanel/SMediaStreamMediaTrack.h
+++ b/Engine/Plugins/Experimental/MediaStream/Source/MediaStreamEditor/Private/DetailsPanel/SMediaStreamMediaTrack.h
@@ -22,7 +22,7 @@ class SMediaStreamMediaTrack : public SCompoundWidget
 		{}
 	SLATE_END_ARGS()
 
-	MEDIASTREAMEDITOR_API void Construct(const FArguments& InArgs, UMediaStream* InMediaStream);
+	MEDIASTREAMEDITOR_API void Construct(const FArguments& InArgs, const TArray<UMediaStream*>& InMediaStreams);
 
 	//~ Begin SWidget
 	MEDIASTREAMEDITOR_API virtual void Tick(const FGeometry& InAllottedGeometry, const double InCurrentTime, const float InDeltaTime) override;
@@ -30,16 +30,16 @@ class SMediaStreamMediaTrack : public SCompoundWidget
 
 private:
 	/** The media stream for this track. */
-	TWeakObjectPtr<UMediaStream> MediaStreamWeak;
+	TArray<TWeakObjectPtr<UMediaStream>> MediaStreamsWeak;
 
 	/** The active media player. If this changes, recreate the track. */
-	TWeakObjectPtr<UMediaPlayer> MediaPlayerWeak;
+	TArray<TWeakObjectPtr<UMediaPlayer>> MediaPlayersWeak;
 
 	/** The container for the track. */
 	TSharedPtr<SBox> Content;
 
 	/** Returns the currently active media player. May different from @see MediaPlayerWeak. */
-	UMediaPlayer* GetMediaPlayer() const;
+	TArray<UMediaPlayer*> GetMediaPlayers() const;
 
 	/** Creates the track widget. */
 	TSharedRef<SWidget> CreateTrack();
diff --git a/Engine/Plugins/Experimental/MediaStream/Source/MediaStreamEditor/Private/DetailsPanel/SMediaStreamPlaybackControls.cpp b/Engine/Plugins/Experimental/MediaStream/Source/MediaStreamEditor/Private/DetailsPanel/SMediaStreamPlaybackControls.cpp
--- a/Engine/Plugins/Experimental/MediaStream/Source/MediaStreamEditor/Private/DetailsPanel/SMediaStreamPlaybackControls.cpp
+++ b/Engine/Plugins/Experimental/MediaStream/Source/MediaStreamEditor/Private/DetailsPanel/SMediaStreamPlaybackControls.cpp
@@ -22,9 +22,14 @@
 namespace UE::MediaStreamEditor
 {
 
-void SMediaStreamPlaybackControls::Construct(const FArguments& InArgs, UMediaStream* InMediaStream)
+void SMediaStreamPlaybackControls::Construct(const FArguments& InArgs, const TArray<UMediaStream*>& InMediaStreams)
 {
-	MediaStreamWeak = InMediaStream;
+	MediaStreamsWeak.Reserve(InMediaStreams.Num());
+
+	for (UMediaStream* MediaStream : InMediaStreams)
+	{
+		MediaStreamsWeak.Add(MediaStream);
+	}
 
 	const ISlateStyle* MediaStreamStyle = &FMediaStreamEditorStyle::Get().Get();
 
@@ -173,39 +178,56 @@ void SMediaStreamPlaybackControls::Construct(const FArguments& InArgs, UMediaStr
 	];
 }
 
-UMediaStream* SMediaStreamPlaybackControls::GetMediaStream() const
-{
-	return MediaStreamWeak.Get();
-}
-
-UMediaPlayer* SMediaStreamPlaybackControls::GetMediaPlayer() const
+TArray<UMediaStream*> SMediaStreamPlaybackControls::GetMediaStreams() const
 {
-	UMediaStream* MediaStream = MediaStreamWeak.Get();
+	TArray<UMediaStream*> MediaStreams;
+	MediaStreams.Reserve(MediaStreamsWeak.Num());
 
-	if (!MediaStream)
+	for (const TWeakObjectPtr<UMediaStream>& MediaStreamWeak : MediaStreamsWeak)
 	{
-		return nullptr;
+		if (UMediaStream* MediaStream = MediaStreamWeak.Get())
+		{
+			MediaStreams.Add(MediaStream);
+		}
 	}
 
-	IMediaStreamPlayer* MediaStreamPlayer = MediaStream->GetPlayer().GetInterface();
+	return MediaStreams;
+}
+
+TArray<UMediaPlayer*> SMediaStreamPlaybackControls::GetMediaPlayers() const
+{
+	TArray<UMediaStream*> MediaStreams = GetMediaStreams();
+	TArray<UMediaPlayer*> MediaPlayers;
+	MediaPlayers.Reserve(MediaStreams.Num());
 
-	if (!MediaStreamPlayer)
+	for (UMediaStream* MediaStream : MediaStreams)
 	{
-		return nullptr;
+		if (IMediaStreamPlayer* MediaStreamPlayer = MediaStream->GetPlayer().GetInterface())
+		{
+			if (UMediaPlayer* MediaPlayer = MediaStreamPlayer->GetPlayer())
+			{
+				MediaPlayers.Add(MediaPlayer);
+			}
+		}
 	}
 
-	return MediaStreamPlayer->GetPlayer();
+	return MediaPlayers;
 }
 
 void SMediaStreamPlaybackControls::OnChangePlaybackState(EMediaStreamPlaybackState InState)
 {
-	if (UMediaStream* MediaStream = GetMediaStream())
+	for (UMediaStream* MediaStream : GetMediaStreams())
 	{
+		if (FMediaStreamEditorSequencerLibrary::HasTrack(MediaStream))
+		{
+			continue;
+		}
+		
 		if (IMediaStreamPlayer* MediaStreamPlayer = MediaStream->GetPlayer().GetInterface())
 		{
 			if (MediaStreamPlayer->IsReadOnly())
 			{
-				return;
+				continue;
 			}
 
 			switch (InState)
@@ -236,13 +258,18 @@ void SMediaStreamPlaybackControls::OnChangePlaybackState(EMediaStreamPlaybackSta
 
 void SMediaStreamPlaybackControls::OnChangePlaybackDirection(EMediaStreamPlaybackDirection InDirection)
 {
-	if (UMediaStream* MediaStream = GetMediaStream())
+	for (UMediaStream* MediaStream : GetMediaStreams())
 	{
+		if (FMediaStreamEditorSequencerLibrary::HasTrack(MediaStream))
+		{
+			continue;
+		}
+
 		if (IMediaStreamPlayer* MediaStreamPlayer = MediaStream->GetPlayer().GetInterface())
 		{
 			if (MediaStreamPlayer->IsReadOnly())
 			{
-				return;
+				continue;
 			}
 
 			const float RequestedRate = FMath::Abs(MediaStreamPlayer->GetPlayerConfig().Rate)
@@ -261,13 +288,18 @@ void SMediaStreamPlaybackControls::OnChangePlaybackDirection(EMediaStreamPlaybac
 
 void SMediaStreamPlaybackControls::OnPlaybackSeek(EMediaStreamPlaybackSeek InPosition)
 {
-	if (UMediaStream* MediaStream = GetMediaStream())
+	for (UMediaStream* MediaStream : GetMediaStreams())
 	{
+		if (FMediaStreamEditorSequencerLibrary::HasTrack(MediaStream))
+		{
+			continue;
+		}
+
 		if (IMediaStreamPlayer* MediaStreamPlayer = MediaStream->GetPlayer().GetInterface())
 		{
 			if (MediaStreamPlayer->IsReadOnly())
 			{
-				return;
+				continue;
 			}
 
 			switch (InPosition)
@@ -294,14 +326,20 @@ void SMediaStreamPlaybackControls::OnPlaybackSeek(EMediaStreamPlaybackSeek InPos
 
 bool SMediaStreamPlaybackControls::IsRateSupported(float InRate) const
 {
-	UMediaPlayer* MediaPlayer = GetMediaPlayer();
+	for (UMediaPlayer* MediaPlayer : GetMediaPlayers())
+	{
+		if (MediaPlayer->SupportsRate(InRate, /* Unthinned */ true))
+		{
+			return true;
+		}
+	}
 
-	return MediaPlayer && MediaPlayer->SupportsRate(InRate, /* Unthinned */ true);
+	return false;
 }
 
 EVisibility SMediaStreamPlaybackControls::Open_GetVisibility() const
 {
-	if (UMediaStream* MediaStream = GetMediaStream())
+	for (UMediaStream* MediaStream : GetMediaStreams())
 	{
 		if (MediaStream->HasValidSource())
 		{
@@ -309,16 +347,16 @@ EVisibility SMediaStreamPlaybackControls::Open_GetVisibility() const
 			{
 				if (UMediaPlayer* MediaPlayer = MediaStreamPlayer->GetPlayer())
 				{
-					if (MediaPlayer->IsClosed())
+					if (!MediaPlayer->IsReady())
 					{
 						return EVisibility::Visible;
 					}
-
-					return EVisibility::Collapsed;
+				}
+				else
+				{
+					return EVisibility::Visible;
 				}
 			}
-
-			return EVisibility::Visible;
 		}
 	}
 
@@ -327,7 +365,7 @@ EVisibility SMediaStreamPlaybackControls::Open_GetVisibility() const
 
 FReply SMediaStreamPlaybackControls::Open_OnClicked()
 {
-	if (UMediaStream* MediaStream = GetMediaStream())
+	for (UMediaStream* MediaStream : GetMediaStreams())
 	{
 		if (IMediaStreamPlayer* MediaStreamPlayer = MediaStream->GetPlayer().GetInterface())
 		{
@@ -340,7 +378,7 @@ FReply SMediaStreamPlaybackControls::Open_OnClicked()
 
 EVisibility SMediaStreamPlaybackControls::Close_GetVisibility() const
 {
-	if (UMediaStream* MediaStream = GetMediaStream())
+	for (UMediaStream* MediaStream : GetMediaStreams())
 	{
 		if (MediaStream->HasValidSource())
 		{
@@ -365,7 +403,7 @@ EVisibility SMediaStreamPlaybackControls::Close_GetVisibility() const
 
 FReply SMediaStreamPlaybackControls::Close_OnClicked()
 {
-	if (UMediaStream* MediaStream = GetMediaStream())
+	for (UMediaStream* MediaStream : GetMediaStreams())
 	{
 		if (IMediaStreamPlayer* MediaStreamPlayer = MediaStream->GetPlayer().GetInterface())
 		{
@@ -378,15 +416,21 @@ FReply SMediaStreamPlaybackControls::Close_OnClicked()
 
 bool SMediaStreamPlaybackControls::Rewind_IsEnabled() const
 {
-	if (UMediaStream* MediaStream = GetMediaStream())
+	for (UMediaStream* MediaStream : GetMediaStreams())
 	{
 		if (!FMediaStreamEditorSequencerLibrary::HasTrack(MediaStream))
 		{
-			if (const UMediaPlayer* MediaPlayer = GetMediaPlayer())
+			if (IMediaStreamPlayer* MediaStreamPlayer = MediaStream->GetPlayer().GetInterface())
 			{
-				return MediaPlayer->IsReady() &&
-					MediaPlayer->SupportsSeeking() &&
-					MediaPlayer->GetTime() > FTimespan::Zero();
+				if (UMediaPlayer* MediaPlayer = MediaStreamPlayer->GetPlayer())
+				{
+					if (MediaPlayer->IsReady() &&
+						MediaPlayer->SupportsSeeking() &&
+						MediaPlayer->GetTime() > FTimespan::Zero())
+					{
+						return true;
+					}
+				}
 			}
 		}
 	}
@@ -402,15 +446,21 @@ FReply SMediaStreamPlaybackControls::Rewind_OnClicked()
 
 bool SMediaStreamPlaybackControls::Reverse_IsEnabled() const
 {
-	if (UMediaStream* MediaStream = GetMediaStream())
+	for (UMediaStream* MediaStream : GetMediaStreams())
 	{
 		if (!FMediaStreamEditorSequencerLibrary::HasTrack(MediaStream))
 		{
-			if (UMediaPlayer* MediaPlayer = GetMediaPlayer())
+			if (IMediaStreamPlayer* MediaStreamPlayer = MediaStream->GetPlayer().GetInterface())
 			{
-				const float Rate = FMath::Abs(MediaPlayer->GetRate());
+				if (UMediaPlayer* MediaPlayer = MediaStreamPlayer->GetPlayer())
+				{
+					const float Rate = FMath::Abs(MediaPlayer->GetRate());
 
-				return MediaPlayer->IsReady() && IsRateSupported(Rate > 0 ? -Rate : -1.f);
+					if (MediaPlayer->IsReady() && IsRateSupported(Rate > 0 ? -Rate : -1.f))
+					{
+						return true;
+					}
+				}
 			}
 		}
 	}
@@ -426,25 +476,42 @@ FReply SMediaStreamPlaybackControls::Reverse_OnClicked()
 
 EVisibility SMediaStreamPlaybackControls::Play_GetVisibility() const
 {
-	if (const UMediaPlayer* MediaPlayer = GetMediaPlayer())
+	for (UMediaStream* MediaStream : GetMediaStreams())
 	{
-		return MediaPlayer->IsPlaying() ? EVisibility::Collapsed : EVisibility::Visible;
+		if (!FMediaStreamEditorSequencerLibrary::HasTrack(MediaStream))
+		{
+			if (IMediaStreamPlayer* MediaStreamPlayer = MediaStream->GetPlayer().GetInterface())
+			{
+				if (UMediaPlayer* MediaPlayer = MediaStreamPlayer->GetPlayer())
+				{
+					if (!MediaPlayer->IsPlaying())
+					{
+						return EVisibility::Visible;
+					}
+				}
+			}
+		}
 	}
 
-	return EVisibility::Visible;
+	return EVisibility::Collapsed;
 }
 
 bool SMediaStreamPlaybackControls::Play_IsEnabled() const
 {
-	if (UMediaStream* MediaStream = GetMediaStream())
+	for (UMediaStream* MediaStream : GetMediaStreams())
 	{
 		if (!FMediaStreamEditorSequencerLibrary::HasTrack(MediaStream))
 		{
-			// Is the player paused or fast forwarding/rewinding?
-			if (const UMediaPlayer* MediaPlayer = GetMediaPlayer())
+			if (IMediaStreamPlayer* MediaStreamPlayer = MediaStream->GetPlayer().GetInterface())
 			{
-				return MediaPlayer->IsReady()
-					&& (!MediaPlayer->IsPlaying() || (MediaPlayer->GetRate() != 1.0f));
+				if (UMediaPlayer* MediaPlayer = MediaStreamPlayer->GetPlayer())
+				{
+					if (MediaPlayer->IsReady()
+						&& (!MediaPlayer->IsPlaying() || (MediaPlayer->GetRate() != 1.0f)))
+					{
+						return true;
+					}
+				}
 			}
 		}
 	}
@@ -460,10 +527,21 @@ FReply SMediaStreamPlaybackControls::Play_OnClicked()
 
 EVisibility SMediaStreamPlaybackControls::Pause_GetVisibility() const
 {
-	if (const UMediaPlayer* MediaPlayer = GetMediaPlayer())
+	for (UMediaStream* MediaStream : GetMediaStreams())
 	{
-		const bool bIsVisible = MediaPlayer->CanPause() && !MediaPlayer->IsPaused();
-		return bIsVisible ? EVisibility::Visible : EVisibility::Collapsed;
+		if (!FMediaStreamEditorSequencerLibrary::HasTrack(MediaStream))
+		{
+			if (IMediaStreamPlayer* MediaStreamPlayer = MediaStream->GetPlayer().GetInterface())
+			{
+				if (UMediaPlayer* MediaPlayer = MediaStreamPlayer->GetPlayer())
+				{
+					if (MediaPlayer->CanPause() && !MediaPlayer->IsPaused())
+					{
+						return EVisibility::Visible;
+					}
+				}
+			}
+		}
 	}
 
 	return EVisibility::Collapsed;
@@ -471,13 +549,19 @@ EVisibility SMediaStreamPlaybackControls::Pause_GetVisibility() const
 
 bool SMediaStreamPlaybackControls::Pause_IsEnabled() const
 {
-	if (UMediaStream* MediaStream = GetMediaStream())
+	for (UMediaStream* MediaStream : GetMediaStreams())
 	{
 		if (!FMediaStreamEditorSequencerLibrary::HasTrack(MediaStream))
 		{
-			if (const UMediaPlayer* MediaPlayer = GetMediaPlayer())
+			if (IMediaStreamPlayer* MediaStreamPlayer = MediaStream->GetPlayer().GetInterface())
 			{
-				return MediaPlayer->CanPause() && !MediaPlayer->IsPaused();
+				if (UMediaPlayer* MediaPlayer = MediaStreamPlayer->GetPlayer())
+				{
+					if (MediaPlayer->CanPause() && !MediaPlayer->IsPaused())
+					{
+						return true;
+					}
+				}
 			}
 		}
 	}
@@ -493,13 +577,19 @@ FReply SMediaStreamPlaybackControls::Pause_OnClicked()
 
 bool SMediaStreamPlaybackControls::Forward_IsEnabled() const
 {
-	if (UMediaStream* MediaStream = GetMediaStream())
+	for (UMediaStream* MediaStream : GetMediaStreams())
 	{
 		if (!FMediaStreamEditorSequencerLibrary::HasTrack(MediaStream))
 		{
-			if (UMediaPlayer* MediaPlayer = GetMediaPlayer())
+			if (IMediaStreamPlayer* MediaStreamPlayer = MediaStream->GetPlayer().GetInterface())
 			{
-				return MediaPlayer->IsReady();
+				if (UMediaPlayer* MediaPlayer = MediaStreamPlayer->GetPlayer())
+				{
+					if (MediaPlayer->IsReady())
+					{
+						return true;
+					}
+				}
 			}
 		}
 	}
@@ -515,9 +605,9 @@ FReply SMediaStreamPlaybackControls::Forward_OnClicked()
 
 EVisibility SMediaStreamPlaybackControls::AddTrack_GetVisibility() const
 {
-	if (UMediaStream* MediaStream = GetMediaStream())
+	for (UMediaStream* MediaStream : GetMediaStreams())
 	{
-		if (FMediaStreamEditorSequencerLibrary::CanAddTrack(MediaStream))
+		if (FMediaStreamEditorSequencerLibrary::CanAddTrack(MediaStream) && !FMediaStreamEditorSequencerLibrary::HasTrack(MediaStream))
 		{
 			return EVisibility::Visible;
 		}
@@ -528,9 +618,12 @@ EVisibility SMediaStreamPlaybackControls::AddTrack_GetVisibility() const
 
 FReply SMediaStreamPlaybackControls::AddTrack_OnClicked()
 {
-	if (UMediaStream* MediaStream = GetMediaStream())
+	for (UMediaStream* MediaStream : GetMediaStreams())
 	{
-		FMediaStreamEditorSequencerLibrary::AddTrack(MediaStream);
+		if (FMediaStreamEditorSequencerLibrary::CanAddTrack(MediaStream))
+		{
+			FMediaStreamEditorSequencerLibrary::AddTrack(MediaStream);
+		}
 	}
 
 	return FReply::Unhandled();
diff --git a/Engine/Plugins/Experimental/MediaStream/Source/MediaStreamEditor/Private/DetailsPanel/SMediaStreamPlaybackControls.h b/Engine/Plugins/Experimental/MediaStream/Source/MediaStreamEditor/Private/DetailsPanel/SMediaStreamPlaybackControls.h
--- a/Engine/Plugins/Experimental/MediaStream/Source/MediaStreamEditor/Private/DetailsPanel/SMediaStreamPlaybackControls.h
+++ b/Engine/Plugins/Experimental/MediaStream/Source/MediaStreamEditor/Private/DetailsPanel/SMediaStreamPlaybackControls.h
@@ -25,17 +25,17 @@ class SMediaStreamPlaybackControls : public SCompoundWidget
 		{}
 	SLATE_END_ARGS()
 
-	MEDIASTREAMEDITOR_API void Construct(const FArguments& InArgs, UMediaStream* InMediaStream);
+	MEDIASTREAMEDITOR_API void Construct(const FArguments& InArgs, const TArray<UMediaStream*>& InMediaStreams);
 
 private:
 	/** The media stream for these controls. */
-	TWeakObjectPtr<UMediaStream> MediaStreamWeak;
+	TArray<TWeakObjectPtr<UMediaStream>> MediaStreamsWeak;
 
-	/** Returns the media stream associated with these controls. */
-	UMediaStream* GetMediaStream() const;
+	/** Returns the media streams associated with these controls. */
+	TArray<UMediaStream*> GetMediaStreams() const;
 
-	/** Returns the media player associated with these controls. */
-	UMediaPlayer* GetMediaPlayer() const;
+	/** Returns the media players associated with these controls. */
+	TArray<UMediaPlayer*> GetMediaPlayers() const;
 
 	/** Changes the state of selected media streams. */
 	void OnChangePlaybackState(EMediaStreamPlaybackState InState);
diff --git a/Engine/Plugins/Experimental/MediaStream/Source/MediaStreamEditor/Private/MediaStreamWidgets.cpp b/Engine/Plugins/Experimental/MediaStream/Source/MediaStreamEditor/Private/MediaStreamWidgets.cpp
--- a/Engine/Plugins/Experimental/MediaStream/Source/MediaStreamEditor/Private/MediaStreamWidgets.cpp
+++ b/Engine/Plugins/Experimental/MediaStream/Source/MediaStreamEditor/Private/MediaStreamWidgets.cpp
@@ -21,14 +21,14 @@ TSharedRef<SWidget> FMediaStreamWidgets::CreateTextureDetailsWidget(UMediaStream
 	return SNew(SMediaStreamMediaDetails, InMediaStream);
 }
 
-TSharedRef<SWidget> FMediaStreamWidgets::CreateTrackWidget(UMediaStream* InMediaStream)
+TSharedRef<SWidget> FMediaStreamWidgets::CreateTrackWidget(const TArray<UMediaStream*>& InMediaStreams)
 {
-	return SNew(SMediaStreamMediaTrack, InMediaStream);
+	return SNew(SMediaStreamMediaTrack, InMediaStreams);
 }
 
-TSharedRef<SWidget> FMediaStreamWidgets::CreateControlsWidget(UMediaStream* InMediaStream)
+TSharedRef<SWidget> FMediaStreamWidgets::CreateControlsWidget(const TArray<UMediaStream*>& InMediaStreams)
 {
-	return SNew(SMediaStreamPlaybackControls, InMediaStream);
+	return SNew(SMediaStreamPlaybackControls, InMediaStreams);
 }
 
 }; // UE::MediaStreamEditor
diff --git a/Engine/Plugins/Experimental/MediaStream/Source/MediaStreamEditor/Public/MediaStreamWidgets.h b/Engine/Plugins/Experimental/MediaStream/Source/MediaStreamEditor/Public/MediaStreamWidgets.h
--- a/Engine/Plugins/Experimental/MediaStream/Source/MediaStreamEditor/Public/MediaStreamWidgets.h
+++ b/Engine/Plugins/Experimental/MediaStream/Source/MediaStreamEditor/Public/MediaStreamWidgets.h
@@ -22,17 +22,18 @@ public:
 	/**
 	 * Creates a widget representing the details of the owned Media Texture.
 	 */
-	MEDIASTREAMEDITOR_API static TSharedRef<SWidget> CreateTextureDetailsWidget(UMediaStream* InMediaStream);
+	MEDIASTREAMEDITOR_API static TSharedRef<SWidget> CreateTextureDetailsWidget(UMediaStream* InMediaStreams);
 
 	/**
 	 * Creates a widget allowing the user to scrub the media.
 	 */
-	MEDIASTREAMEDITOR_API static TSharedRef<SWidget> CreateTrackWidget(UMediaStream* InMediaStream);
+	MEDIASTREAMEDITOR_API static TSharedRef<SWidget> CreateTrackWidget(const TArray<UMediaStream*>& InMediaStreams);
 
 	/**
 	 * Creates a widget allowing the user to control the media playback.
+	 * If providing multiple media streams, the first should be the one that can have tracks.
 	 */
-	MEDIASTREAMEDITOR_API static TSharedRef<SWidget> CreateControlsWidget(UMediaStream* InMediaStream);
+	MEDIASTREAMEDITOR_API static TSharedRef<SWidget> CreateControlsWidget(const TArray<UMediaStream*>& InMediaStreams);
 };
 
 } // UE::MediaStreamEditor
