Media Stream: Sequencer display
- Added object schema and default track name.

diff --git a/Engine/Plugins/Experimental/MediaStream/Source/MediaStreamEditor/MediaStreamEditor.Build.cs b/Engine/Plugins/Experimental/MediaStream/Source/MediaStreamEditor/MediaStreamEditor.Build.cs
--- a/Engine/Plugins/Experimental/MediaStream/Source/MediaStreamEditor/MediaStreamEditor.Build.cs
+++ b/Engine/Plugins/Experimental/MediaStream/Source/MediaStreamEditor/MediaStreamEditor.Build.cs
@@ -31,6 +31,7 @@ public class MediaStreamEditor : ModuleRules
 				"MovieScene",
 				"Projects",
 				"PropertyEditor",
+				"Sequencer",
 				"Slate",
 				"SlateCore",
 				"ToolMenus",
diff --git a/Engine/Plugins/Experimental/MediaStream/Source/MediaStreamEditor/Private/MediaStreamEditorModule.cpp b/Engine/Plugins/Experimental/MediaStream/Source/MediaStreamEditor/Private/MediaStreamEditorModule.cpp
--- a/Engine/Plugins/Experimental/MediaStream/Source/MediaStreamEditor/Private/MediaStreamEditorModule.cpp
+++ b/Engine/Plugins/Experimental/MediaStream/Source/MediaStreamEditor/Private/MediaStreamEditorModule.cpp
@@ -3,15 +3,24 @@
 #include "MediaStreamEditorModule.h"
 
 #include "DetailsPanel/MediaStreamCustomization.h"
+#include "ISequencerModule.h"
 #include "MediaStream.h"
 #include "MediaStreamEditorStyle.h"
 #include "MediaStreamEditorToolbar.h"
+#include "MediaStreamObjectSchema.h"
 #include "Modules/ModuleManager.h"
 #include "PropertyEditorDelegates.h"
 #include "PropertyEditorModule.h"
 
+namespace UE::MediaStream::Private
+{
+	const FLazyName SequencerModuleName = TEXT("Sequencer");
+}
+
 void FMediaStreamEditorModule::StartupModule()
 {
+	using namespace UE::MediaStream::Private;
+
 	FMediaStreamEditorStyle::Get();
 
 	FPropertyEditorModule& PropertyModule = FModuleManager::LoadModuleChecked<FPropertyEditorModule>("PropertyEditor");
@@ -20,11 +29,26 @@ void FMediaStreamEditorModule::StartupModule()
 		FOnGetDetailCustomizationInstance::CreateStatic(&UE::MediaStreamEditor::FMediaStreamCustomization::MakeInstance));
 
 	FMediaStreamEditorToolbar::Get().Integrate();
+
+	ISequencerModule& SequencerModule = FModuleManager::Get().GetModuleChecked<ISequencerModule>(SequencerModuleName);
+
+	MediaStreamObjectSchema = MakeShared<FMediaStreamObjectSchema>();
+	SequencerModule.RegisterObjectSchema(MediaStreamObjectSchema);
 }
 
 void FMediaStreamEditorModule::ShutdownModule()
 {
+	using namespace UE::MediaStream::Private;
+
 	FMediaStreamEditorToolbar::Get().Disintegrate();
+
+	if (MediaStreamObjectSchema.IsValid())
+	{
+		if (ISequencerModule* SequencerModule = FModuleManager::Get().GetModulePtr<ISequencerModule>(SequencerModuleName))
+		{
+			SequencerModule->UnregisterObjectSchema(MediaStreamObjectSchema);
+		}
+	}
 }
 
 IMPLEMENT_MODULE(FMediaStreamEditorModule, MediaStreamEditor);
diff --git a/Engine/Plugins/Experimental/MediaStream/Source/MediaStreamEditor/Private/MediaStreamEditorModule.h b/Engine/Plugins/Experimental/MediaStream/Source/MediaStreamEditor/Private/MediaStreamEditorModule.h
--- a/Engine/Plugins/Experimental/MediaStream/Source/MediaStreamEditor/Private/MediaStreamEditorModule.h
+++ b/Engine/Plugins/Experimental/MediaStream/Source/MediaStreamEditor/Private/MediaStreamEditorModule.h
@@ -5,6 +5,9 @@
 #include "Modules/ModuleInterface.h"
 
 #include "Delegates/IDelegateInstance.h"
+#include "Templates/SharedPointer.h"
+
+class FMediaStreamObjectSchema;
 
 /**
  * Media Stream Editor - Content/type agnostic chainable media proxy with media player integration.
@@ -16,4 +19,7 @@ public:
 	virtual void StartupModule() override;
 	virtual void ShutdownModule() override;
 	//~ End IModuleInterface
+
+protected:
+	TSharedPtr<FMediaStreamObjectSchema> MediaStreamObjectSchema;
 };
diff --git a/Engine/Plugins/Experimental/MediaStream/Source/MediaStreamEditor/Private/MediaStreamEditorSequencerLibrary.cpp b/Engine/Plugins/Experimental/MediaStream/Source/MediaStreamEditor/Private/MediaStreamEditorSequencerLibrary.cpp
--- a/Engine/Plugins/Experimental/MediaStream/Source/MediaStreamEditor/Private/MediaStreamEditorSequencerLibrary.cpp
+++ b/Engine/Plugins/Experimental/MediaStream/Source/MediaStreamEditor/Private/MediaStreamEditorSequencerLibrary.cpp
@@ -11,6 +11,8 @@
 #include "MovieSceneMediaSection.h"
 #include "MovieSceneMediaTrack.h"
 
+#define LOCTEXT_NAMESPACE "MediaStreamEditorSequencerLibrary"
+
 ULevelSequence* FMediaStreamEditorSequencerLibrary::GetLevelSequence()
 {
 	if (ULevelSequence* FocusedSequence = ULevelSequenceEditorBlueprintLibrary::GetFocusedLevelSequence())
@@ -224,5 +226,9 @@ void FMediaStreamEditorSequencerLibrary::AddTrack(UMediaStream* InMediaStream)
 	MediaSection->SetEndFrame(Range.GetUpperBound());
 	MediaSection->SetMediaSourceProxy(UE::MovieScene::FRelativeObjectBindingID(MediaStreamId), /* Proxy Index */ 0);
 
+	MediaTrack->SetDisplayName(LOCTEXT("MediaTrack", "Media Track"));
+
 	MediaTrack->AddSection(*MediaSection);
 }
+
+#undef LOCTEXT_NAMESPACE
diff --git a/Engine/Plugins/Experimental/MediaStream/Source/MediaStreamEditor/Private/MediaStreamObjectSchema.cpp b/Engine/Plugins/Experimental/MediaStream/Source/MediaStreamEditor/Private/MediaStreamObjectSchema.cpp
new file mode 100644
--- /dev/null
+++ b/Engine/Plugins/Experimental/MediaStream/Source/MediaStreamEditor/Private/MediaStreamObjectSchema.cpp
@@ -0,0 +1,32 @@
+// Copyright Epic Games, Inc. All Rights Reserved.
+
+#include "MediaStreamObjectSchema.h"
+
+#include "GameFramework/Actor.h"
+#include "MediaStream.h"
+
+UObject* FMediaStreamObjectSchema::GetParentObject(UObject* InObject) const
+{
+	if (InObject)
+	{
+		if (AActor* Actor = InObject->GetTypedOuter<AActor>())
+		{
+			return Actor;
+		}
+	}
+
+	return nullptr;
+}
+
+UE::Sequencer::FObjectSchemaRelevancy FMediaStreamObjectSchema::GetRelevancy(const UObject* InObject) const
+{
+	// Return a non-zero priority so that it wins against the default relevancy
+	return UE::Sequencer::FObjectSchemaRelevancy(UMediaStream::StaticClass(), 1000);
+}
+
+TSharedPtr<FExtender> FMediaStreamObjectSchema::ExtendObjectBindingMenu(TSharedRef<FUICommandList> InCommandList
+	, TWeakPtr<ISequencer> InSequencerWeak
+	, TConstArrayView<UObject*> InContextSensitiveObjects) const
+{
+	return nullptr;
+}
\ No newline at end of file
diff --git a/Engine/Plugins/Experimental/MediaStream/Source/MediaStreamEditor/Private/MediaStreamObjectSchema.h b/Engine/Plugins/Experimental/MediaStream/Source/MediaStreamEditor/Private/MediaStreamObjectSchema.h
new file mode 100644
--- /dev/null
+++ b/Engine/Plugins/Experimental/MediaStream/Source/MediaStreamEditor/Private/MediaStreamObjectSchema.h
@@ -0,0 +1,15 @@
+// Copyright Epic Games, Inc. All Rights Reserved.
+
+#pragma once
+
+#include "ISequencerObjectSchema.h"
+
+class FMediaStreamObjectSchema : public UE::Sequencer::IObjectSchema
+{
+public:
+	//~ Begin IObjectSchema
+	virtual UObject* GetParentObject(UObject* InObject) const override;
+	virtual UE::Sequencer::FObjectSchemaRelevancy GetRelevancy(const UObject* InObject) const override;
+	virtual TSharedPtr<FExtender> ExtendObjectBindingMenu(TSharedRef<FUICommandList> InCommandList, TWeakPtr<ISequencer> InSequencerWeak, TConstArrayView<UObject*> InContextSensitiveObjects) const override;
+	//~ End IObjectSchema
+};
