Media Stream: Updated control view.
- Controls now use a wrapbox to display
- Changed button visuals to match media viewer / sequencer.

diff --git a/Engine/Plugins/Experimental/MediaStream/Source/MediaStreamEditor/Private/DetailsPanel/SMediaStreamPlaybackControls.cpp b/Engine/Plugins/Experimental/MediaStream/Source/MediaStreamEditor/Private/DetailsPanel/SMediaStreamPlaybackControls.cpp
--- a/Engine/Plugins/Experimental/MediaStream/Source/MediaStreamEditor/Private/DetailsPanel/SMediaStreamPlaybackControls.cpp
+++ b/Engine/Plugins/Experimental/MediaStream/Source/MediaStreamEditor/Private/DetailsPanel/SMediaStreamPlaybackControls.cpp
@@ -14,6 +14,7 @@
 #include "Styling/StyleColors.h"
 #include "Widgets/Images/SImage.h"
 #include "Widgets/Input/SButton.h"
+#include "Widgets/Layout/SWrapBox.h"
 #include "Widgets/SBoxPanel.h"
 #include "Widgets/Text/STextBlock.h"
 
@@ -35,17 +36,21 @@ void SMediaStreamPlaybackControls::Construct(const FArguments& InArgs, const TAr
 
 	ChildSlot
 	[
-		SNew(SHorizontalBox)
+		SNew(SWrapBox)
+		.Orientation(Orient_Horizontal)
+		.UseAllottedSize(true)
+		.HAlign(HAlign_Center)
+		.InnerSlotPadding(FVector2D::ZeroVector)
 
 		// Open button.
-		+ SHorizontalBox::Slot()
+		+ SWrapBox::Slot()
 		.VAlign(VAlign_Center)
-		.AutoWidth()
 		[
 			SNew(SButton)
 			.VAlign(VAlign_Center)
 			.Visibility(this, &SMediaStreamPlaybackControls::Open_GetVisibility)
 			.OnClicked(this, &SMediaStreamPlaybackControls::Open_OnClicked)
+			.ButtonStyle(FMediaStreamEditorStyle::Get(), "MediaButtons")
 			[
 				SNew(SImage)
 				.ColorAndOpacity(FSlateColor::UseForeground())
@@ -55,14 +60,14 @@ void SMediaStreamPlaybackControls::Construct(const FArguments& InArgs, const TAr
 		]
 
 		// Close button.
-		+ SHorizontalBox::Slot()
+		+ SWrapBox::Slot()
 		.VAlign(VAlign_Center)
-		.AutoWidth()
 		[
 			SNew(SButton)
 			.VAlign(VAlign_Center)
 			.Visibility(this, &SMediaStreamPlaybackControls::Close_GetVisibility)
 			.OnClicked(this, &SMediaStreamPlaybackControls::Close_OnClicked)
+			.ButtonStyle(FMediaStreamEditorStyle::Get(), "MediaButtons")
 			[
 				SNew(SImage)
 				.ColorAndOpacity(FSlateColor::UseForeground())
@@ -72,14 +77,14 @@ void SMediaStreamPlaybackControls::Construct(const FArguments& InArgs, const TAr
 		]
 
 		// Rewind button.
-		+ SHorizontalBox::Slot()
+		+ SWrapBox::Slot()
 		.VAlign(VAlign_Center)
-		.AutoWidth()
 		[
 			SNew(SButton)
 			.VAlign(VAlign_Center)
 			.IsEnabled(this, &SMediaStreamPlaybackControls::Rewind_IsEnabled)
 			.OnClicked(this, &SMediaStreamPlaybackControls::Rewind_OnClicked)
+			.ButtonStyle(FMediaStreamEditorStyle::Get(), "MediaButtons")
 			[
 				SNew(SImage)
 				.ColorAndOpacity(FSlateColor::UseForeground())
@@ -89,14 +94,14 @@ void SMediaStreamPlaybackControls::Construct(const FArguments& InArgs, const TAr
 		]
 
 		// Reverse button.
-		+ SHorizontalBox::Slot()
+		+ SWrapBox::Slot()
 		.VAlign(VAlign_Center)
-		.AutoWidth()
 		[
 			SNew(SButton)
 			.VAlign(VAlign_Center)
 			.IsEnabled(this, &SMediaStreamPlaybackControls::Reverse_IsEnabled)
 			.OnClicked(this, &SMediaStreamPlaybackControls::Reverse_OnClicked)
+			.ButtonStyle(FMediaStreamEditorStyle::Get(), "MediaButtons")
 			[
 				SNew(SImage)
 				.ColorAndOpacity(FSlateColor::UseForeground())
@@ -106,15 +111,15 @@ void SMediaStreamPlaybackControls::Construct(const FArguments& InArgs, const TAr
 		]
 
 		// Play button.
-		+ SHorizontalBox::Slot()
+		+ SWrapBox::Slot()
 		.VAlign(VAlign_Center)
-		.AutoWidth()
 		[
 			SNew(SButton)
 			.VAlign(VAlign_Center)
 			.Visibility(this, &SMediaStreamPlaybackControls::Play_GetVisibility)
 			.IsEnabled(this, &SMediaStreamPlaybackControls::Play_IsEnabled)
 			.OnClicked(this, &SMediaStreamPlaybackControls::Play_OnClicked)
+			.ButtonStyle(FMediaStreamEditorStyle::Get(), "MediaButtons")
 			[
 				SNew(SImage)
 				.ColorAndOpacity(FSlateColor::UseForeground())
@@ -124,15 +129,15 @@ void SMediaStreamPlaybackControls::Construct(const FArguments& InArgs, const TAr
 		]
 
 		// Pause button.
-		+ SHorizontalBox::Slot()
+		+ SWrapBox::Slot()
 		.VAlign(VAlign_Center)
-		.AutoWidth()
 		[
 			SNew(SButton)
 			.VAlign(VAlign_Center)
 			.Visibility(this, &SMediaStreamPlaybackControls::Pause_GetVisibility)
 			.IsEnabled(this, &SMediaStreamPlaybackControls::Pause_IsEnabled)
 			.OnClicked(this, &SMediaStreamPlaybackControls::Pause_OnClicked)
+			.ButtonStyle(FMediaStreamEditorStyle::Get(), "MediaButtons")
 			[
 				SNew(SImage)
 				.ColorAndOpacity(FSlateColor::UseForeground())
@@ -142,14 +147,14 @@ void SMediaStreamPlaybackControls::Construct(const FArguments& InArgs, const TAr
 		]
 
 		// Forward button.
-		+ SHorizontalBox::Slot()
+		+ SWrapBox::Slot()
 		.VAlign(VAlign_Center)
-		.AutoWidth()
 		[
 			SNew(SButton)
 			.VAlign(VAlign_Center)
 			.IsEnabled(this, &SMediaStreamPlaybackControls::Forward_IsEnabled)
 			.OnClicked(this, &SMediaStreamPlaybackControls::Forward_OnClicked)
+			.ButtonStyle(FMediaStreamEditorStyle::Get(), "MediaButtons")
 			[
 				SNew(SImage)
 				.ColorAndOpacity(FSlateColor::UseForeground())
@@ -160,15 +165,15 @@ void SMediaStreamPlaybackControls::Construct(const FArguments& InArgs, const TAr
 		]
 
 		// Add Media Track
-		+ SHorizontalBox::Slot()
+		+ SWrapBox::Slot()
 		.VAlign(VAlign_Center)
-		.AutoWidth()
 		[
 			SNew(SButton)
 			.ContentPadding(FMargin(0.f, 4.f, 0.f, 3.f))
 			.VAlign(VAlign_Center)
 			.Visibility(this, &SMediaStreamPlaybackControls::AddTrack_GetVisibility)
 			.OnClicked(this, &SMediaStreamPlaybackControls::AddTrack_OnClicked)
+			.ButtonStyle(FMediaStreamEditorStyle::Get(), "MediaButtons")
 			[
 				SNew(STextBlock)
 				.Font(IDetailLayoutBuilder::GetDetailFont())
diff --git a/Engine/Plugins/Experimental/MediaStream/Source/MediaStreamEditor/Private/MediaStreamEditorStyle.cpp b/Engine/Plugins/Experimental/MediaStream/Source/MediaStreamEditor/Private/MediaStreamEditorStyle.cpp
--- a/Engine/Plugins/Experimental/MediaStream/Source/MediaStreamEditor/Private/MediaStreamEditorStyle.cpp
+++ b/Engine/Plugins/Experimental/MediaStream/Source/MediaStreamEditor/Private/MediaStreamEditorStyle.cpp
@@ -9,6 +9,7 @@
 #include "Interfaces/IPluginManager.h"
 #include "Math/Vector2D.h"
 #include "Misc/Paths.h"
+#include "Styling/AppStyle.h"
 #include "Styling/SlateStyleMacros.h"
 #include "Styling/SlateStyleRegistry.h"
 #include "Styling/SlateTypes.h"
@@ -73,6 +74,14 @@ FMediaStreamEditorStyle::FMediaStreamEditorStyle()
 		.SetBarThickness(2.0f)
 	);
 
+	Set("MediaButtons", FButtonStyle(FAppStyle::Get().GetWidgetStyle<FButtonStyle>("Animation.PlayControlsButton"))
+		.SetNormal(FSlateNoResource())
+		.SetDisabled(FSlateNoResource())
+		.SetHovered(FSlateRoundedBoxBrush(FLinearColor(0.2, 0.2, 0.2, 0.5), 3.f, FVector2D(20.f)))
+		.SetPressed(FSlateRoundedBoxBrush(FLinearColor(0.1, 0.1, 0.1, 0.5), 3.f, FVector2D(20.f)))
+		.SetNormalPadding(FMargin(2.f, 2.f, 2.f, 2.f))
+		.SetPressedPadding(FMargin(2.f, 2.f, 2.f, 2.f)));
+
 	FSlateStyleRegistry::RegisterSlateStyle(*this);
 }
 
