Media Stream: Changed load button to console variable.

diff --git a/Engine/Plugins/Experimental/MediaStream/Source/MediaStream/Private/MediaStreamModule.cpp b/Engine/Plugins/Experimental/MediaStream/Source/MediaStream/Private/MediaStreamModule.cpp
--- a/Engine/Plugins/Experimental/MediaStream/Source/MediaStream/Private/MediaStreamModule.cpp
+++ b/Engine/Plugins/Experimental/MediaStream/Source/MediaStream/Private/MediaStreamModule.cpp
@@ -2,6 +2,7 @@
 
 #include "MediaStreamModule.h"
 
+#include "HAL/IConsoleManager.h"
 #include "MediaStreamObjectHandlerManager.h"
 #include "MediaStreamSchemeHandlerManager.h"
 #include "Modules/ModuleManager.h"
@@ -20,6 +21,17 @@
 
 DEFINE_LOG_CATEGORY(LogMediaStream);
 
+#if WITH_EDITOR
+namespace UE::MediaStream
+{
+	TAutoConsoleVariable<bool> CVarAutoLoadMediaStreamsOnMapLoad = TAutoConsoleVariable<bool>(
+		TEXT("MediaStream.OpenOnMapLoad"),
+		/* Default */ true,
+		TEXT("Automatically loads media streams when loading a level. If set to false, media streams must be manually loaded. Only applies to the Level Editor, not PIE or runtime.")
+	);
+}
+#endif
+
 FMediaStreamModule& FMediaStreamModule::Get()
 {
 	return FModuleManager::Get().GetModuleChecked<FMediaStreamModule>(UE_MODULE_NAME);
@@ -28,7 +40,9 @@ FMediaStreamModule& FMediaStreamModule::Get()
 bool FMediaStreamModule::CanOpenSourceOnLoad()
 {
 #if WITH_EDITOR
-	return !GIsEditor || !bIsMapLoading;
+	using namespace UE::MediaStream;
+
+	return !GIsEditor || !bIsMapLoading || CVarAutoLoadMediaStreamsOnMapLoad.GetValueOnAnyThread();
 #else
 	return true;
 #endif
@@ -37,7 +51,9 @@ bool FMediaStreamModule::CanOpenSourceOnLoad()
 bool FMediaStreamModule::CanAutoplay() const
 {
 #if WITH_EDITOR
-	return !GIsEditor || !bIsMapLoading;
+	using namespace UE::MediaStream;
+
+	return !GIsEditor || !bIsMapLoading || CVarAutoLoadMediaStreamsOnMapLoad.GetValueOnAnyThread();
 #else
 	return true;
 #endif	
diff --git a/Engine/Plugins/Experimental/MediaStream/Source/MediaStreamEditor/Private/MediaStreamEditorModule.cpp b/Engine/Plugins/Experimental/MediaStream/Source/MediaStreamEditor/Private/MediaStreamEditorModule.cpp
--- a/Engine/Plugins/Experimental/MediaStream/Source/MediaStreamEditor/Private/MediaStreamEditorModule.cpp
+++ b/Engine/Plugins/Experimental/MediaStream/Source/MediaStreamEditor/Private/MediaStreamEditorModule.cpp
@@ -6,7 +6,6 @@
 #include "ISequencerModule.h"
 #include "MediaStream.h"
 #include "MediaStreamEditorStyle.h"
-#include "MediaStreamEditorToolbar.h"
 #include "MediaStreamObjectSchema.h"
 #include "Modules/ModuleManager.h"
 #include "PropertyEditorDelegates.h"
@@ -28,8 +27,6 @@ void FMediaStreamEditorModule::StartupModule()
 	PropertyModule.RegisterCustomClassLayout(UMediaStream::StaticClass()->GetFName(),
 		FOnGetDetailCustomizationInstance::CreateStatic(&UE::MediaStreamEditor::FMediaStreamCustomization::MakeInstance));
 
-	FMediaStreamEditorToolbar::Get().Integrate();
-
 	ISequencerModule& SequencerModule = FModuleManager::Get().GetModuleChecked<ISequencerModule>(SequencerModuleName);
 
 	MediaStreamObjectSchema = MakeShared<FMediaStreamObjectSchema>();
@@ -40,8 +37,6 @@ void FMediaStreamEditorModule::ShutdownModule()
 {
 	using namespace UE::MediaStream::Private;
 
-	FMediaStreamEditorToolbar::Get().Disintegrate();
-
 	if (MediaStreamObjectSchema.IsValid())
 	{
 		if (ISequencerModule* SequencerModule = FModuleManager::Get().GetModulePtr<ISequencerModule>(SequencerModuleName))
diff --git a/Engine/Plugins/Experimental/MediaStream/Source/MediaStreamEditor/Private/MediaStreamEditorToolbar.cpp b/Engine/Plugins/Experimental/MediaStream/Source/MediaStreamEditor/Private/MediaStreamEditorToolbar.cpp
deleted file mode 100644
diff --git a/Engine/Plugins/Experimental/MediaStream/Source/MediaStreamEditor/Private/MediaStreamEditorToolbar.h b/Engine/Plugins/Experimental/MediaStream/Source/MediaStreamEditor/Private/MediaStreamEditorToolbar.h
deleted file mode 100644
