Material Designer / Media Stream: Fixed remote control integration for source selector.

diff --git a/Engine/Plugins/Experimental/DynamicMaterial/Source/DynamicMaterialEditor/Private/DynamicMaterialEditorModule.cpp b/Engine/Plugins/Experimental/DynamicMaterial/Source/DynamicMaterialEditor/Private/DynamicMaterialEditorModule.cpp
--- a/Engine/Plugins/Experimental/DynamicMaterial/Source/DynamicMaterialEditor/Private/DynamicMaterialEditorModule.cpp
+++ b/Engine/Plugins/Experimental/DynamicMaterial/Source/DynamicMaterialEditor/Private/DynamicMaterialEditorModule.cpp
@@ -393,6 +393,11 @@ void FDynamicMaterialEditorModule::ClearDynamicMaterialModel(UWorld* InWorld) co
 	OpenMaterialModel(nullptr, InWorld, /* Invoke tab */ false);
 }
 
+IDMWidgetLibrary& FDynamicMaterialEditorModule::GetWidgetLibrary() const
+{
+	return FDMWidgetLibrary::Get();
+}
+
 TSharedRef<SWidget> FDynamicMaterialEditorModule::CreateEditor(UDynamicMaterialModelBase* InMaterialModelBase, UWorld* InAssetEditorWorld)
 {
 	TSharedRef<SDMMaterialDesigner> NewDesigner = SNew(SDMMaterialDesigner);
diff --git a/Engine/Plugins/Experimental/DynamicMaterial/Source/DynamicMaterialEditor/Private/DynamicMaterialEditorModule.h b/Engine/Plugins/Experimental/DynamicMaterial/Source/DynamicMaterialEditor/Private/DynamicMaterialEditorModule.h
--- a/Engine/Plugins/Experimental/DynamicMaterial/Source/DynamicMaterialEditor/Private/DynamicMaterialEditorModule.h
+++ b/Engine/Plugins/Experimental/DynamicMaterial/Source/DynamicMaterialEditor/Private/DynamicMaterialEditorModule.h
@@ -73,6 +73,7 @@ public:
 	virtual void OpenMaterial(UDynamicMaterialInstance* InMaterial, UWorld* InWorld, bool bInInvokeTab) const override;
 	virtual void OnActorSelected(AActor* InActor, UWorld* InWorld, bool bInInvokeTab) const override;
 	virtual void ClearDynamicMaterialModel(UWorld* InWorld) const override;
+	virtual IDMWidgetLibrary& GetWidgetLibrary() const override;
 	//~ End IDynamicMaterialEditorModule
 
 	//~ Begin IModuleInterface
diff --git a/Engine/Plugins/Experimental/DynamicMaterial/Source/DynamicMaterialEditor/Private/UI/Utils/DMWidgetLibrary.cpp b/Engine/Plugins/Experimental/DynamicMaterial/Source/DynamicMaterialEditor/Private/UI/Utils/DMWidgetLibrary.cpp
--- a/Engine/Plugins/Experimental/DynamicMaterial/Source/DynamicMaterialEditor/Private/UI/Utils/DMWidgetLibrary.cpp
+++ b/Engine/Plugins/Experimental/DynamicMaterial/Source/DynamicMaterialEditor/Private/UI/Utils/DMWidgetLibrary.cpp
@@ -188,6 +188,7 @@ FDMPropertyHandle FDMWidgetLibrary::CreatePropertyHandle(const FDMPropertyHandle
 
 	FPropertyRowGeneratorArgs PreviewRowGeneratorArgs;
 	PreviewRowGeneratorArgs.NotifyHook = InParams.NotifyHook;
+	PreviewRowGeneratorArgs.bShouldShowHiddenProperties = true;
 
 	PropertyHandle.PreviewHandle.PropertyRowGenerator = PropertyEditor.CreatePropertyRowGenerator(PreviewRowGeneratorArgs);
 	PropertyHandle.PreviewHandle.PropertyRowGenerator->SetObjects({InParams.Object});
@@ -217,6 +218,7 @@ FDMPropertyHandle FDMWidgetLibrary::CreatePropertyHandle(const FDMPropertyHandle
 		{
 			// No notify hook for the original
 			FPropertyRowGeneratorArgs OriginalRowGeneratorArgs;
+			OriginalRowGeneratorArgs.bShouldShowHiddenProperties = true;
 
 			PropertyHandle.OriginalHandle.PropertyRowGenerator = PropertyEditor.CreatePropertyRowGenerator(OriginalRowGeneratorArgs);
 			PropertyHandle.OriginalHandle.PropertyRowGenerator->SetObjects({OriginalObject});
diff --git a/Engine/Plugins/Experimental/DynamicMaterial/Source/DynamicMaterialEditor/Public/IDynamicMaterialEditorModule.h b/Engine/Plugins/Experimental/DynamicMaterial/Source/DynamicMaterialEditor/Public/IDynamicMaterialEditorModule.h
--- a/Engine/Plugins/Experimental/DynamicMaterial/Source/DynamicMaterialEditor/Public/IDynamicMaterialEditorModule.h
+++ b/Engine/Plugins/Experimental/DynamicMaterial/Source/DynamicMaterialEditor/Public/IDynamicMaterialEditorModule.h
@@ -109,4 +109,6 @@ public:
 	virtual void OnActorSelected(AActor* InActor, UWorld* InWorld, bool bInInvokeTab) const = 0;
 
 	virtual void ClearDynamicMaterialModel(UWorld* InWorld) const = 0;
+	
+	virtual IDMWidgetLibrary& GetWidgetLibrary() const = 0;
 };
diff --git a/Engine/Plugins/Experimental/DynamicMaterialMediaStreamBridge/Source/DynamicMaterialMediaStreamBridgeEditor/Private/DMMaterialValueMediaStreamPropertyRowGenerator.cpp b/Engine/Plugins/Experimental/DynamicMaterialMediaStreamBridge/Source/DynamicMaterialMediaStreamBridgeEditor/Private/DMMaterialValueMediaStreamPropertyRowGenerator.cpp
--- a/Engine/Plugins/Experimental/DynamicMaterialMediaStreamBridge/Source/DynamicMaterialMediaStreamBridgeEditor/Private/DMMaterialValueMediaStreamPropertyRowGenerator.cpp
+++ b/Engine/Plugins/Experimental/DynamicMaterialMediaStreamBridge/Source/DynamicMaterialMediaStreamBridgeEditor/Private/DMMaterialValueMediaStreamPropertyRowGenerator.cpp
@@ -134,12 +134,15 @@ void FDMMaterialValueMediaStreamPropertyRowGenerator::AddSourceCategory(FDMCompo
 			continue;
 		}
 
-		FDMPropertyHandle& Handle = InParams.PropertyRows->AddDefaulted_GetRef();
+		FDMPropertyHandle Handle = IDynamicMaterialEditorModule::Get().GetWidgetLibrary().GetPropertyHandle(InParams.CreatePropertyHandleParams(WidgetRow.SourceProperty->GetFName()));
 		Handle.bEnabled = WidgetRow.Enabled.Get();
 		Handle.bKeyframeable = false;
-		Handle.ValueName = *WidgetRow.Name.ToString();
+		Handle.NameOverride = WidgetRow.SourceProperty->GetDisplayNameText();
+		Handle.ValueName = *WidgetRow.Name.ToString();                           
 		Handle.ValueWidget = WidgetRow.Widget;
 		Handle.CategoryOverrideName = CategoryName;
+
+		InParams.PropertyRows-> Add(Handle);
 	}
 }
 
diff --git a/Engine/Plugins/Experimental/MediaStream/Source/MediaStream/Private/SchemeHandlers/MediaStreamAssetSchemeHandler.cpp b/Engine/Plugins/Experimental/MediaStream/Source/MediaStream/Private/SchemeHandlers/MediaStreamAssetSchemeHandler.cpp
--- a/Engine/Plugins/Experimental/MediaStream/Source/MediaStream/Private/SchemeHandlers/MediaStreamAssetSchemeHandler.cpp
+++ b/Engine/Plugins/Experimental/MediaStream/Source/MediaStream/Private/SchemeHandlers/MediaStreamAssetSchemeHandler.cpp
@@ -123,7 +123,8 @@ void FMediaStreamAssetSchemeHandler::AddAssetSelector(UMediaStream* InMediaStrea
 		.DisplayThumbnail(true)
 		.AllowedClass(UMediaSource::StaticClass()),
 		/* Enabled */ true,
-		TAttribute<EVisibility>::CreateSP(this, &FMediaStreamAssetSchemeHandler::GetAssetSelectorVisibility, MediaStreamWeak)
+		TAttribute<EVisibility>::CreateSP(this, &FMediaStreamAssetSchemeHandler::GetAssetSelectorVisibility, MediaStreamWeak),
+		FMediaStreamSource::StaticStruct()->FindPropertyByName(GET_MEMBER_NAME_CHECKED(FMediaStreamSource, Object))
 	});
 }
 
diff --git a/Engine/Plugins/Experimental/MediaStream/Source/MediaStream/Private/SchemeHandlers/MediaStreamFileSchemeHandler.cpp b/Engine/Plugins/Experimental/MediaStream/Source/MediaStream/Private/SchemeHandlers/MediaStreamFileSchemeHandler.cpp
--- a/Engine/Plugins/Experimental/MediaStream/Source/MediaStream/Private/SchemeHandlers/MediaStreamFileSchemeHandler.cpp
+++ b/Engine/Plugins/Experimental/MediaStream/Source/MediaStream/Private/SchemeHandlers/MediaStreamFileSchemeHandler.cpp
@@ -152,7 +152,8 @@ void FMediaStreamFileSchemeHandler::AddFileSelector(UMediaStream* InMediaStream,
 		.FileTypeFilter(TEXT("All files (*.*)|*.*"))
 		.OnPathPicked(this, &FMediaStreamFileSchemeHandler::OnFilePicked, MediaStreamWeak),
 		/* Enabled */ true,
-		TAttribute<EVisibility>::CreateSP(this, &FMediaStreamFileSchemeHandler::GetFileSelectorVisibility, MediaStreamWeak)
+		TAttribute<EVisibility>::CreateSP(this, &FMediaStreamFileSchemeHandler::GetFileSelectorVisibility, MediaStreamWeak),
+		FMediaStreamSource::StaticStruct()->FindPropertyByName(GET_MEMBER_NAME_CHECKED(FMediaStreamSource, Path))
 	});
 }
 
diff --git a/Engine/Plugins/Experimental/MediaStream/Source/MediaStream/Public/IMediaStreamSchemeHandler.h b/Engine/Plugins/Experimental/MediaStream/Source/MediaStream/Public/IMediaStreamSchemeHandler.h
--- a/Engine/Plugins/Experimental/MediaStream/Source/MediaStream/Public/IMediaStreamSchemeHandler.h
+++ b/Engine/Plugins/Experimental/MediaStream/Source/MediaStream/Public/IMediaStreamSchemeHandler.h
@@ -79,6 +79,7 @@ public:
 		TSharedRef<SWidget> Widget;
 		TAttribute<bool> Enabled = true;
 		TAttribute<EVisibility> Visibility = EVisibility::Visible;
+		FProperty* SourceProperty = nullptr;
 	};
 
 	struct FCustomWidgets
diff --git a/Engine/Plugins/Experimental/MediaStream/Source/MediaStream/Public/MediaStreamSource.h b/Engine/Plugins/Experimental/MediaStream/Source/MediaStream/Public/MediaStreamSource.h
--- a/Engine/Plugins/Experimental/MediaStream/Source/MediaStream/Public/MediaStreamSource.h
+++ b/Engine/Plugins/Experimental/MediaStream/Source/MediaStream/Public/MediaStreamSource.h
@@ -29,7 +29,7 @@ struct FMediaStreamSource
 	 * Used by the handler subsystem to identifer the scheme used.
 	 * A value of NAME_None identifies a null or invalid source.
 	 */
-	UPROPERTY(BlueprintReadOnly, Category = "Media Stream|Source",
+	UPROPERTY(EditAnywhere, BlueprintReadOnly, Category = "Media Stream|Source",
 		meta = (GetOptions = "/Script/MediaStream.MediaStreamSourceBlueprintFunctionLibrary:GetSchemeTypes"))
 	FName Scheme = NAME_None;
 
@@ -37,13 +37,13 @@ struct FMediaStreamSource
 	 * Path to the referenced source, such as a file/asset path or a managed source name.
 	 * A path must be provided to have a valid scheme.
 	 */
-	UPROPERTY(BlueprintReadOnly, Category = "Media Stream|Source")
+	UPROPERTY(EditAnywhere, BlueprintReadOnly, Category = "Media Stream|Source")
 	FString Path;
 
 	/**
 	 * Source of the Media. Either a UMediaStream, UMediaSource or a UMediaPlaylist
 	 */
-	UPROPERTY(BlueprintReadOnly, Category = "Media Stream|Source",
+	UPROPERTY(EditAnywhere, BlueprintReadOnly, Category = "Media Stream|Source",
 		meta = (AllowedClasses="MediaStream,MediaSource,MediaPlaylist"))
 	TObjectPtr<UObject> Object = nullptr;
 
diff --git a/Engine/Plugins/Experimental/MediaStream/Source/MediaStreamEditor/Private/DetailsPanel/MediaStreamCustomization.cpp b/Engine/Plugins/Experimental/MediaStream/Source/MediaStreamEditor/Private/DetailsPanel/MediaStreamCustomization.cpp
--- a/Engine/Plugins/Experimental/MediaStream/Source/MediaStreamEditor/Private/DetailsPanel/MediaStreamCustomization.cpp
+++ b/Engine/Plugins/Experimental/MediaStream/Source/MediaStreamEditor/Private/DetailsPanel/MediaStreamCustomization.cpp
@@ -158,9 +158,18 @@ void FMediaStreamCustomization::AddSourceCategory(IDetailLayoutBuilder& InDetail
 
 	for (const IMediaStreamSchemeHandler::FWidgetRow& WidgetRow : WidgetRows.CustomRows)
 	{
-		SourceCategory.AddCustomRow(WidgetRow.Name)
+		if (!WidgetRow.SourceProperty)
+		{
+			continue;
+		}
+
+		TSharedPtr<IPropertyHandle> ChildHandle = SourceHandle->GetChildHandle(WidgetRow.SourceProperty->GetFName());
+		
+		SourceCategory.AddProperty(ChildHandle)
 			.IsEnabled(WidgetRow.Enabled)
 			.Visibility(WidgetRow.Visibility)
+			.DisplayName(WidgetRow.Name)
+			.CustomWidget()
 			.NameContent()
 			[
 				SNew(STextBlock)
diff --git a/Engine/Plugins/Experimental/MediaStream/Source/MediaStreamEditor/Private/DetailsPanel/MediaStreamSourceCustomization.cpp b/Engine/Plugins/Experimental/MediaStream/Source/MediaStreamEditor/Private/DetailsPanel/MediaStreamSourceCustomization.cpp
--- a/Engine/Plugins/Experimental/MediaStream/Source/MediaStreamEditor/Private/DetailsPanel/MediaStreamSourceCustomization.cpp
+++ b/Engine/Plugins/Experimental/MediaStream/Source/MediaStreamEditor/Private/DetailsPanel/MediaStreamSourceCustomization.cpp
@@ -59,8 +59,11 @@ void FMediaStreamSourceCustomization::AddSourceSchemeSelector(IMediaStreamScheme
 	 */
 
 	InOutCustomWidgets.CustomRows.Add({
-		LOCTEXT("SourceSchemeName", "Source Scheme"),
-		SchemeSelector
+		LOCTEXT("SourceSchemeName", "Scheme"),
+		SchemeSelector,
+		/* Enabled */ true,
+		EVisibility::Visible,
+		FMediaStreamSource::StaticStruct()->FindPropertyByName(GET_MEMBER_NAME_CHECKED(FMediaStreamSource, Scheme))
 	});
 }
 
