Fix short names in AllowedClasses metadata.

diff --git a/Engine/Plugins/Experimental/MediaStream/Source/MediaStream/Public/MediaStreamSource.h b/Engine/Plugins/Experimental/MediaStream/Source/MediaStream/Public/MediaStreamSource.h
--- a/Engine/Plugins/Experimental/MediaStream/Source/MediaStream/Public/MediaStreamSource.h
+++ b/Engine/Plugins/Experimental/MediaStream/Source/MediaStream/Public/MediaStreamSource.h
@@ -44,7 +44,7 @@ struct FMediaStreamSource
 	 * Source of the Media. Either a UMediaStream, UMediaSource or a UMediaPlaylist
 	 */
 	UPROPERTY(EditAnywhere, BlueprintReadOnly, Category = "Media Stream|Source",
-		meta = (AllowedClasses="MediaStream,MediaSource,MediaPlaylist"))
+		meta = (AllowedClasses="/Script/MediaStream.MediaStream,/Script/MediaAssets.MediaSource,/Script/MediaAssets.MediaPlaylist"))
 	TObjectPtr<UObject> Object = nullptr;
 
 	MEDIASTREAM_API bool operator==(const FMediaStreamSource& InOther) const;
