Media Stream: Corrected "Add Track" button for Media Sources
- Will now always display button, but disable it,r ather than hide it, when it is not available.
- Added tooltip to button explaining what it does and when it is usable.
- It is now disabled if the Media Stream's world is not the level level sequence's world.

diff --git a/Engine/Plugins/Experimental/MediaStream/Source/MediaStreamEditor/MediaStreamEditor.Build.cs b/Engine/Plugins/Experimental/MediaStream/Source/MediaStreamEditor/MediaStreamEditor.Build.cs
--- a/Engine/Plugins/Experimental/MediaStream/Source/MediaStreamEditor/MediaStreamEditor.Build.cs
+++ b/Engine/Plugins/Experimental/MediaStream/Source/MediaStreamEditor/MediaStreamEditor.Build.cs
@@ -35,6 +35,7 @@ public class MediaStreamEditor : ModuleRules
 				"Slate",
 				"SlateCore",
 				"ToolMenus",
+				"UnrealEd",
 			}
 		);
 	}
diff --git a/Engine/Plugins/Experimental/MediaStream/Source/MediaStreamEditor/Private/DetailsPanel/SMediaStreamPlaybackControls.cpp b/Engine/Plugins/Experimental/MediaStream/Source/MediaStreamEditor/Private/DetailsPanel/SMediaStreamPlaybackControls.cpp
--- a/Engine/Plugins/Experimental/MediaStream/Source/MediaStreamEditor/Private/DetailsPanel/SMediaStreamPlaybackControls.cpp
+++ b/Engine/Plugins/Experimental/MediaStream/Source/MediaStreamEditor/Private/DetailsPanel/SMediaStreamPlaybackControls.cpp
@@ -169,11 +169,12 @@ void SMediaStreamPlaybackControls::Construct(const FArguments& InArgs, const TAr
 		.VAlign(VAlign_Center)
 		[
 			SNew(SButton)
-			.ContentPadding(FMargin(0.f, 4.f, 0.f, 3.f))
+			.ContentPadding(FMargin(4.f, 4.f, 4.f, 3.f))
 			.VAlign(VAlign_Center)
-			.Visibility(this, &SMediaStreamPlaybackControls::AddTrack_GetVisibility)
+			.IsEnabled(this, &SMediaStreamPlaybackControls::AddTrack_GetEnabled)
 			.OnClicked(this, &SMediaStreamPlaybackControls::AddTrack_OnClicked)
 			.ButtonStyle(FMediaStreamEditorStyle::Get(), "MediaButtons")
+			.ToolTipText(this, &SMediaStreamPlaybackControls::AddTrack_GetToolTip)
 			[
 				SNew(STextBlock)
 				.Font(IDetailLayoutBuilder::GetDetailFont())
@@ -608,17 +609,69 @@ FReply SMediaStreamPlaybackControls::Forward_OnClicked()
 	return FReply::Handled();
 }
 
-EVisibility SMediaStreamPlaybackControls::AddTrack_GetVisibility() const
+bool SMediaStreamPlaybackControls::AddTrack_GetEnabled() const
 {
 	for (UMediaStream* MediaStream : GetMediaStreams())
 	{
-		if (FMediaStreamEditorSequencerLibrary::CanAddTrack(MediaStream) && !FMediaStreamEditorSequencerLibrary::HasTrack(MediaStream))
+		if (!FMediaStreamEditorSequencerLibrary::HasTrack(MediaStream) && FMediaStreamEditorSequencerLibrary::CanAddTrack(MediaStream))
 		{
-			return EVisibility::Visible;
+			return true;
 		}
 	}
 
-	return EVisibility::Collapsed;
+	return false;
+}
+
+FText SMediaStreamPlaybackControls::AddTrack_GetToolTip() const
+{
+	bool bCanAddTrack = false;
+	bool bHasTrack = false;
+	bool bHasAsset = false;
+
+	for (UMediaStream* MediaStream : GetMediaStreams())
+	{
+		if (!bCanAddTrack && FMediaStreamEditorSequencerLibrary::CanAddTrack(MediaStream))
+		{
+			bCanAddTrack = true;
+		}
+
+		if (!bHasTrack && FMediaStreamEditorSequencerLibrary::HasTrack(MediaStream))
+		{
+			bHasTrack = true;
+		}
+
+		if (!bHasAsset && !MediaStream->GetWorld() && !MediaStream->IsIn(GetTransientPackage()))
+		{
+			bHasAsset = true;
+		}
+	}
+
+	TArray<FText> ToolTips;
+
+	if (!bCanAddTrack)
+	{
+		ToolTips.Add(LOCTEXT("CannotAddTrackToolTip", "Note: A track cannot be added."));
+	}
+
+	if (bHasTrack)
+	{
+		ToolTips.Add(LOCTEXT("HasTrackToolTip", "Note: A track has already been added."));
+	}
+
+	if (bHasAsset)
+	{
+		ToolTips.Add(LOCTEXT("HasAssetToolTip", "Note: You cannot bind assets from the Content Browser to the Level Sequence. Consider creating a Material Designer Instance."));
+	}
+
+	if (ToolTips.IsEmpty())
+	{
+		return LOCTEXT("AddTrackToolTip", "Add Track to Level Sequence");
+	}
+
+	return FText::Format(
+		LOCTEXT("AddTrackWithNotesToolTip", "Add Track to Level Sequence\n\n{0}"),
+		FText::Join(INVTEXT("\n"), ToolTips)
+	);
 }
 
 FReply SMediaStreamPlaybackControls::AddTrack_OnClicked()
diff --git a/Engine/Plugins/Experimental/MediaStream/Source/MediaStreamEditor/Private/DetailsPanel/SMediaStreamPlaybackControls.h b/Engine/Plugins/Experimental/MediaStream/Source/MediaStreamEditor/Private/DetailsPanel/SMediaStreamPlaybackControls.h
--- a/Engine/Plugins/Experimental/MediaStream/Source/MediaStreamEditor/Private/DetailsPanel/SMediaStreamPlaybackControls.h
+++ b/Engine/Plugins/Experimental/MediaStream/Source/MediaStreamEditor/Private/DetailsPanel/SMediaStreamPlaybackControls.h
@@ -74,7 +74,8 @@ private:
 	bool Forward_IsEnabled() const;
 	FReply Forward_OnClicked();
 
-	EVisibility AddTrack_GetVisibility() const;
+	bool AddTrack_GetEnabled() const;
+	FText AddTrack_GetToolTip() const;
 	FReply AddTrack_OnClicked();
 };
 
diff --git a/Engine/Plugins/Experimental/MediaStream/Source/MediaStreamEditor/Private/MediaStreamEditorSequencerLibrary.cpp b/Engine/Plugins/Experimental/MediaStream/Source/MediaStreamEditor/Private/MediaStreamEditorSequencerLibrary.cpp
--- a/Engine/Plugins/Experimental/MediaStream/Source/MediaStreamEditor/Private/MediaStreamEditorSequencerLibrary.cpp
+++ b/Engine/Plugins/Experimental/MediaStream/Source/MediaStreamEditor/Private/MediaStreamEditorSequencerLibrary.cpp
@@ -2,6 +2,7 @@
 
 #include "MediaStreamEditorSequencerLibrary.h"
 
+#include "ILevelSequenceEditorToolkit.h"
 #include "IMediaStreamPlayer.h"
 #include "LevelSequence.h"
 #include "LevelSequenceEditorBlueprintLibrary.h"
@@ -10,6 +11,10 @@
 #include "MovieSceneCommonHelpers.h"
 #include "MovieSceneMediaSection.h"
 #include "MovieSceneMediaTrack.h"
+#include "SequencerUtilities.h"
+#include "Subsystems/AssetEditorSubsystem.h"
+
+extern UNREALED_API UEditorEngine* GEditor;
 
 #define LOCTEXT_NAMESPACE "MediaStreamEditorSequencerLibrary"
 
@@ -122,6 +127,11 @@ bool FMediaStreamEditorSequencerLibrary::CanAddTrack(UMediaStream* InMediaStream
 		return false;
 	}
 
+	if (InMediaStream->GetWorld() != World)
+	{
+		return false;
+	}
+
 	TSharedRef<const UE::MovieScene::FSharedPlaybackState> SharedPlaybackState = MovieSceneHelpers::CreateTransientSharedPlaybackState(World, LevelSequence);
 	const FGuid MediaStreamId = LevelSequence->FindBindingFromObject(InMediaStream, SharedPlaybackState);
 
@@ -133,46 +143,51 @@ bool FMediaStreamEditorSequencerLibrary::CanAddTrack(UMediaStream* InMediaStream
 	return LevelSequence->CanPossessObject(*InMediaStream, World);
 }
 
-void FMediaStreamEditorSequencerLibrary::AddTrack(UMediaStream* InMediaStream)
+bool FMediaStreamEditorSequencerLibrary::AddTrack(UMediaStream* InMediaStream)
 {
 	if (!InMediaStream || !CanAddTrack(InMediaStream))
 	{
-		return;
+		return false;
 	}
 
 	IMediaStreamPlayer* MediaStreamPlayerIntf = InMediaStream->GetPlayer().GetInterface();
 
 	if (!MediaStreamPlayerIntf || MediaStreamPlayerIntf->IsReadOnly())
 	{
-		return;
+		return false;
 	}
 
 	UMediaPlayer* MediaPlayer = MediaStreamPlayerIntf->GetPlayer();
 
 	if (!MediaPlayer)
 	{
-		return;
+		return false;
 	}
 
 	ULevelSequence* LevelSequence = GetLevelSequence();
 
 	if (!LevelSequence)
 	{
-		return;
+		return false;
 	}
 
 	UMovieScene* MovieScene = LevelSequence->GetMovieScene();
 
 	if (!MovieScene)
 	{
-		return;
+		return false;
 	}
 
 	UWorld* World = LevelSequence->GetWorld();
 
 	if (!World)
 	{
-		return;
+		return false;
+	}
+
+	if (InMediaStream->GetWorld() != World)
+	{
+		return false;
 	}
 
 	TSharedRef<const UE::MovieScene::FSharedPlaybackState> SharedPlaybackState = MovieSceneHelpers::CreateTransientSharedPlaybackState(World, LevelSequence);
@@ -182,29 +197,34 @@ void FMediaStreamEditorSequencerLibrary::AddTrack(UMediaStream* InMediaStream)
 
 	if (!MediaStreamId.IsValid())
 	{
-		const FString DefaultBindingName = TEXT("MediaStream");
-		const FMovieScenePossessable Possessable(DefaultBindingName, UMediaStream::StaticClass());
-		const FMovieSceneBinding Binding(Possessable.GetGuid(), DefaultBindingName);
-		MovieScene->AddPossessable(Possessable, Binding);
-		MediaStreamId = Possessable.GetGuid();
-		LevelSequence->BindPossessableObject(MediaStreamId, *InMediaStream, World);
+		UAssetEditorSubsystem* AssetEditorSubsystem = GEditor->GetEditorSubsystem<UAssetEditorSubsystem>();
+		IAssetEditorInstance* AssetEditor = AssetEditorSubsystem->FindEditorForAsset(LevelSequence, /* Focus */ false);
+		ILevelSequenceEditorToolkit* LevelSequenceEditor = static_cast<ILevelSequenceEditorToolkit*>(AssetEditor);
+		TSharedPtr<ISequencer> Sequencer = LevelSequenceEditor ? LevelSequenceEditor->GetSequencer() : nullptr;
+
+		if (!Sequencer.IsValid())
+		{
+			return false;
+		}
+
+		MediaStreamId = FSequencerUtilities::CreateBinding(Sequencer.ToSharedRef(), *InMediaStream);
 
 		if (!MediaStreamId.IsValid())
 		{
-			return;
+			return false;
 		}
 	}
 
-	if (MovieScene->FindTrack<UMovieSceneMediaTrack>(MediaStreamId))
+ 	if (MovieScene->FindTrack<UMovieSceneMediaTrack>(MediaStreamId))
 	{
-		return;
+		return false;
 	}
 
 	UMovieSceneMediaTrack* MediaTrack = MovieScene->AddTrack<UMovieSceneMediaTrack>(MediaStreamId);
 
 	if (!MediaTrack)
 	{
-		return;
+		return false;
 	}
 
 	TRange<FFrameNumber> Range = MovieScene->GetPlaybackRange();
@@ -229,6 +249,8 @@ void FMediaStreamEditorSequencerLibrary::AddTrack(UMediaStream* InMediaStream)
 	MediaTrack->SetDisplayName(LOCTEXT("MediaTrack", "Media Track"));
 
 	MediaTrack->AddSection(*MediaSection);
+
+	return true;
 }
 
 #undef LOCTEXT_NAMESPACE
diff --git a/Engine/Plugins/Experimental/MediaStream/Source/MediaStreamEditor/Private/MediaStreamEditorSequencerLibrary.h b/Engine/Plugins/Experimental/MediaStream/Source/MediaStreamEditor/Private/MediaStreamEditorSequencerLibrary.h
--- a/Engine/Plugins/Experimental/MediaStream/Source/MediaStreamEditor/Private/MediaStreamEditorSequencerLibrary.h
+++ b/Engine/Plugins/Experimental/MediaStream/Source/MediaStreamEditor/Private/MediaStreamEditorSequencerLibrary.h
@@ -10,10 +10,10 @@ class FMediaStreamEditorSequencerLibrary
 {
 public:
 	static ULevelSequence* GetLevelSequence();
-
+	
 	static bool HasTrack(UMediaStream* InMediaStream);
 
 	static bool CanAddTrack(UMediaStream* InMediaStream);
 
-	static void AddTrack(UMediaStream* InMediaStream);
+	static bool AddTrack(UMediaStream* InMediaStream);
 };
diff --git a/Engine/Plugins/Experimental/MediaStream/Source/MediaStreamEditor/Private/MediaStreamObjectSchema.cpp b/Engine/Plugins/Experimental/MediaStream/Source/MediaStreamEditor/Private/MediaStreamObjectSchema.cpp
--- a/Engine/Plugins/Experimental/MediaStream/Source/MediaStreamEditor/Private/MediaStreamObjectSchema.cpp
+++ b/Engine/Plugins/Experimental/MediaStream/Source/MediaStreamEditor/Private/MediaStreamObjectSchema.cpp
@@ -5,6 +5,8 @@
 #include "GameFramework/Actor.h"
 #include "MediaStream.h"
 
+#define LOCTEXT_NAMESPACE "MediaStreamObjectSchema"
+
 UObject* FMediaStreamObjectSchema::GetParentObject(UObject* InObject) const
 {
 	if (InObject)
@@ -29,4 +31,11 @@ TSharedPtr<FExtender> FMediaStreamObjectSchema::ExtendObjectBindingMenu(TSharedR
 	, TConstArrayView<UObject*> InContextSensitiveObjects) const
 {
 	return nullptr;
-}
\ No newline at end of file
+}
+
+FText FMediaStreamObjectSchema::GetPrettyName(const UObject* InObject) const
+{
+	return LOCTEXT("Media", "Media");
+}
+
+#undef LOCTEXT_NAMESPACE
diff --git a/Engine/Plugins/Experimental/MediaStream/Source/MediaStreamEditor/Private/MediaStreamObjectSchema.h b/Engine/Plugins/Experimental/MediaStream/Source/MediaStreamEditor/Private/MediaStreamObjectSchema.h
--- a/Engine/Plugins/Experimental/MediaStream/Source/MediaStreamEditor/Private/MediaStreamObjectSchema.h
+++ b/Engine/Plugins/Experimental/MediaStream/Source/MediaStreamEditor/Private/MediaStreamObjectSchema.h
@@ -11,5 +11,6 @@ public:
 	virtual UObject* GetParentObject(UObject* InObject) const override;
 	virtual UE::Sequencer::FObjectSchemaRelevancy GetRelevancy(const UObject* InObject) const override;
 	virtual TSharedPtr<FExtender> ExtendObjectBindingMenu(TSharedRef<FUICommandList> InCommandList, TWeakPtr<ISequencer> InSequencerWeak, TConstArrayView<UObject*> InContextSensitiveObjects) const override;
+	virtual FText GetPrettyName(const UObject* InObject) const;
 	//~ End IObjectSchema
 };
