Media Stream: Fixed sequencer binding issues

diff --git a/Engine/Plugins/Experimental/MediaStream/Source/MediaStreamEditor/Private/MediaStreamObjectSchema.cpp b/Engine/Plugins/Experimental/MediaStream/Source/MediaStreamEditor/Private/MediaStreamObjectSchema.cpp
--- a/Engine/Plugins/Experimental/MediaStream/Source/MediaStreamEditor/Private/MediaStreamObjectSchema.cpp
+++ b/Engine/Plugins/Experimental/MediaStream/Source/MediaStreamEditor/Private/MediaStreamObjectSchema.cpp
@@ -9,6 +9,12 @@
 
 UObject* FMediaStreamObjectSchema::GetParentObject(UObject* InObject) const
 {
+	/*
+	 * Need to re-enable this, but fix the binding bugs first.
+	 * Using this schema changes the binding path to be based on the actor,
+	 * but the context of the binding remainds the world, so it fails to
+	 * resolve.
+	 *
 	if (InObject)
 	{
 		if (AActor* Actor = InObject->GetTypedOuter<AActor>())
@@ -16,6 +22,7 @@ UObject* FMediaStreamObjectSchema::GetParentObject(UObject* InObject) const
 			return Actor;
 		}
 	}
+	 */
 
 	return nullptr;
 }
