Media Stream: Fixed relevancy issue for Sequencer.

diff --git a/Engine/Plugins/Experimental/MediaStream/Source/MediaStreamEditor/Private/MediaStreamObjectSchema.cpp b/Engine/Plugins/Experimental/MediaStream/Source/MediaStreamEditor/Private/MediaStreamObjectSchema.cpp
--- a/Engine/Plugins/Experimental/MediaStream/Source/MediaStreamEditor/Private/MediaStreamObjectSchema.cpp
+++ b/Engine/Plugins/Experimental/MediaStream/Source/MediaStreamEditor/Private/MediaStreamObjectSchema.cpp
@@ -29,8 +29,12 @@ UObject* FMediaStreamObjectSchema::GetParentObject(UObject* InObject) const
 
 UE::Sequencer::FObjectSchemaRelevancy FMediaStreamObjectSchema::GetRelevancy(const UObject* InObject) const
 {
-	// Return a non-zero priority so that it wins against the default relevancy
-	return UE::Sequencer::FObjectSchemaRelevancy(UMediaStream::StaticClass(), 1000);
+	if (InObject && InObject->IsA<UMediaStream>())
+	{
+		return UMediaStream::StaticClass();
+	}
+
+	return UE::Sequencer::FObjectSchemaRelevancy();
 }
 
 TSharedPtr<FExtender> FMediaStreamObjectSchema::ExtendObjectBindingMenu(TSharedRef<FUICommandList> InCommandList
