Media Stream: Fixed for packaged games

diff --git a/Engine/Plugins/Experimental/MediaStream/Source/MediaStream/Private/ObjectHandlers/MediaStreamMediaSourceHandler.cpp b/Engine/Plugins/Experimental/MediaStream/Source/MediaStream/Private/ObjectHandlers/MediaStreamMediaSourceHandler.cpp
--- a/Engine/Plugins/Experimental/MediaStream/Source/MediaStream/Private/ObjectHandlers/MediaStreamMediaSourceHandler.cpp
+++ b/Engine/Plugins/Experimental/MediaStream/Source/MediaStream/Private/ObjectHandlers/MediaStreamMediaSourceHandler.cpp
@@ -64,10 +64,13 @@ UMediaPlayer* FMediaStreamMediaSourceHandler::CreateOrUpdatePlayer(const FMediaS
 	UMediaPlayer* MediaPlayer = NewObject<UMediaPlayer>(InParams.MediaStream, NAME_None, RF_Transactional);
 	const FMediaPlayerOptions Options = MediaStreamPlayer->GetPlayerConfig().CreateOptions(MediaStreamPlayer->GetRequestedSeekTime());
 
-	if (!MediaPlayer->OpenSourceWithOptions(MediaSource, Options))
+	if (!IsRunningCommandlet())
 	{
-		UE_LOG(LogMediaStream, Error, TEXT("Unable to create player for Media Source in FMediaStreamMediaSourceHandler::CreateOrUpdatePlayer"));
-		return nullptr;
+		if (!MediaPlayer->OpenSourceWithOptions(MediaSource, Options))
+		{
+			UE_LOG(LogMediaStream, Error, TEXT("Unable to create player for Media Source in FMediaStreamMediaSourceHandler::CreateOrUpdatePlayer"));
+			return nullptr;
+		}
 	}
 
 	return MediaPlayer;
diff --git a/Engine/Plugins/Experimental/MediaStream/Source/MediaStream/Private/Players/MediaStreamLocalPlayer.cpp b/Engine/Plugins/Experimental/MediaStream/Source/MediaStream/Private/Players/MediaStreamLocalPlayer.cpp
--- a/Engine/Plugins/Experimental/MediaStream/Source/MediaStream/Private/Players/MediaStreamLocalPlayer.cpp
+++ b/Engine/Plugins/Experimental/MediaStream/Source/MediaStream/Private/Players/MediaStreamLocalPlayer.cpp
@@ -3,6 +3,9 @@
 #include "Players/MediaStreamLocalPlayer.h"
 
 #include "CoreGlobals.h"
+#include "IMediaClock.h"
+#include "IMediaClockSink.h"
+#include "IMediaModule.h"
 #include "IMediaStreamSchemeHandler.h"
 #include "LevelSequence.h"
 #include "MediaPlaylist.h"
@@ -10,11 +13,78 @@
 #include "MediaStreamModule.h"
 #include "MediaStreamSchemeHandlerManager.h"
 #include "MediaTexture.h"
+#include "Modules/ModuleManager.h"
 #include "MovieScene.h"
 #include "MovieSceneMediaSection.h"
 #include "MovieSceneTrack.h"
 #include "UObject/UObjectIterator.h"
 
+
+/**
+ * MediaClockSink Callback utility class for the MediaStreamLocalPlayer. 
+ */
+class FMediaStreamLocalPlayerClockSink : public IMediaClockSink
+{
+public:
+	FMediaStreamLocalPlayerClockSink(UMediaStreamLocalPlayer& InPlayer)
+		: PlayerWeak(&InPlayer)
+	{ }
+
+	virtual ~FMediaStreamLocalPlayerClockSink() override = default;
+
+	/** Returns the media module's media clock. */
+	static IMediaClock* GetMediaClock()
+	{
+		if (IMediaModule* MediaModule = FModuleManager::LoadModulePtr<IMediaModule>("Media"))
+		{
+			return &MediaModule->GetClock();
+		}
+		return nullptr;
+	}
+
+	/** Register the clock sink with media module. */
+	static void RegisterClockSink(UMediaStreamLocalPlayer& InPlayer)
+	{
+		if (!InPlayer.ClockSink.IsValid())
+		{
+			if (IMediaClock* MediaClock = GetMediaClock())
+			{
+				InPlayer.ClockSink = MakeShared<FMediaStreamLocalPlayerClockSink>(InPlayer);
+				MediaClock->AddSink(InPlayer.ClockSink.ToSharedRef());
+			}
+		}
+	}
+
+	/** Unregister the clock sink from media module. */
+	static void UnregisterClockSink(UMediaStreamLocalPlayer& InPlayer)
+	{
+		if (InPlayer.ClockSink.IsValid())
+		{
+			if (IMediaClock* MediaClock = GetMediaClock())
+			{
+				MediaClock->RemoveSink(InPlayer.ClockSink.ToSharedRef());
+			}
+			InPlayer.ClockSink.Reset();
+		}
+	}
+	
+	//~ Begin IMediaClockSink
+	virtual void TickOutput(FTimespan InDeltaTime, FTimespan InTimecode) override
+	{
+		if (UMediaStreamLocalPlayer* Player = PlayerWeak.Get())
+		{
+			Player->Initialize();
+			
+			// Because this callback is only to delay player initialization, we can remove it when it is done.
+			UnregisterClockSink(*Player);
+		}
+	}
+	//~ End IMediaClockSink
+
+private:
+	TWeakObjectPtr<UMediaStreamLocalPlayer> PlayerWeak;
+};
+
 UMediaStreamLocalPlayer::~UMediaStreamLocalPlayer()
 {
 	if (!UObjectInitialized())
@@ -124,14 +194,16 @@ void UMediaStreamLocalPlayer::PostDuplicate(bool bInDuplicateForPIE)
 	bIsPIE = bInDuplicateForPIE;
 #endif
 
-	Initialize();
+	// Delay initialization to the next game loop to let all media objects PostLoad.
+	FMediaStreamLocalPlayerClockSink::RegisterClockSink(*this);
 }
 
 void UMediaStreamLocalPlayer::PostEditImport()
 {
 	Super::PostEditImport();
 
-	Initialize();
+	// Delay initialization to the next game loop to let all media objects PostLoad.
+	FMediaStreamLocalPlayerClockSink::RegisterClockSink(*this);
 }
 
 void UMediaStreamLocalPlayer::PostLoad()
@@ -140,14 +212,16 @@ void UMediaStreamLocalPlayer::PostLoad()
 
 	bAllowOpenSource = FMediaStreamModule::Get().CanOpenSourceOnLoad();
 
-	Initialize();
+	// Delay initialization to the next game loop to let all media objects PostLoad.
+	FMediaStreamLocalPlayerClockSink::RegisterClockSink(*this);
 }
 
 void UMediaStreamLocalPlayer::PostNetReceive()
 {
 	Super::PostNetReceive();
 
-	Initialize();
+	// Delay initialization to the next game loop to let all media objects PostLoad.
+	FMediaStreamLocalPlayerClockSink::RegisterClockSink(*this);
 }
 
 void UMediaStreamLocalPlayer::BeginDestroy()
@@ -557,8 +631,6 @@ void UMediaStreamLocalPlayer::InitPlayer()
 
 	if (IsValid(MediaTexture))
 	{
-		MediaTexture->SetMediaPlayer(bIsPlayerValid ? Player : nullptr);
-
 #if WITH_EDITOR
 		MediaTexture->SetDefaultMediaPlayer(bIsPlayerValid ? Player : nullptr);
 #endif
@@ -591,6 +663,13 @@ void UMediaStreamLocalPlayer::DeinitPlayer()
 
 void UMediaStreamLocalPlayer::OnMediaOpened(FString InOpenedUrl)
 {
+	// Connect the media texture's sample queue to the opened player.
+	// This ensures it doesn't get overriden by another player.
+	if (IsValid(Player) && IsValid(MediaTexture))
+	{
+		MediaTexture->SetMediaPlayer(Player);
+	}
+
 	ApplyPlayerConfig();
 }
 
diff --git a/Engine/Plugins/Experimental/MediaStream/Source/MediaStream/Private/Players/MediaStreamLocalPlayer.h b/Engine/Plugins/Experimental/MediaStream/Source/MediaStream/Private/Players/MediaStreamLocalPlayer.h
--- a/Engine/Plugins/Experimental/MediaStream/Source/MediaStream/Private/Players/MediaStreamLocalPlayer.h
+++ b/Engine/Plugins/Experimental/MediaStream/Source/MediaStream/Private/Players/MediaStreamLocalPlayer.h
@@ -10,6 +10,8 @@
 
 #include "MediaStreamLocalPlayer.generated.h"
 
+class FMediaStreamLocalPlayerClockSink;
+
 /**
  * Media Stream Local Player. Plays a media source using a UMediaPlayer.
  */
@@ -46,8 +48,6 @@ public:
 		SetPlaylistIndex(InIndex);
 	}
 
-	void LoadSource();
-
 	//~ Begin IMediaStreamPlayer
 	virtual UMediaStream* GetMediaStream() const override;
 	virtual bool IsReadOnly() const override;
@@ -142,6 +142,13 @@ protected:
 	bool bIsPIE = false;
 #endif
 
+	friend FMediaStreamLocalPlayerClockSink;
+	
+	/** 
+	 * Holds our clock sink if available. 
+	 */
+	TSharedPtr<FMediaStreamLocalPlayerClockSink> ClockSink;
+	
 	/**
 	 * Initializes all of the components of the Media Stream.
 	 */
