Media Stream: Removed Blueprint nodes.

diff --git a/Engine/Plugins/Experimental/MediaStream/Source/MediaStream/Private/Players/MediaStreamProxyPlayer.h b/Engine/Plugins/Experimental/MediaStream/Source/MediaStream/Private/Players/MediaStreamProxyPlayer.h
--- a/Engine/Plugins/Experimental/MediaStream/Source/MediaStream/Private/Players/MediaStreamProxyPlayer.h
+++ b/Engine/Plugins/Experimental/MediaStream/Source/MediaStream/Private/Players/MediaStreamProxyPlayer.h
@@ -18,10 +18,10 @@ class UMediaStreamProxyPlayer : public UObject, public IMediaStreamPlayer
 public:
 	virtual ~UMediaStreamProxyPlayer() override;
 
-	UFUNCTION(BlueprintPure, Category = "Media Stream|Player")
+	UFUNCTION(Category = "Media Stream|Player")
 	TSoftObjectPtr<UMediaStream> GetProxyStreamSoft() const;
 
-	UFUNCTION(BlueprintCallable, Category = "Media Stream|Texture")
+	UFUNCTION(Category = "Media Stream|Texture")
 	void SetReadOnly(bool bInReadOnly);
 
 	//~ Begin IMediaStreamPlayer
diff --git a/Engine/Plugins/Experimental/MediaStream/Source/MediaStream/Public/IMediaStreamPlayer.h b/Engine/Plugins/Experimental/MediaStream/Source/MediaStream/Public/IMediaStreamPlayer.h
--- a/Engine/Plugins/Experimental/MediaStream/Source/MediaStream/Public/IMediaStreamPlayer.h
+++ b/Engine/Plugins/Experimental/MediaStream/Source/MediaStream/Public/IMediaStreamPlayer.h
@@ -34,14 +34,14 @@ class IMediaStreamPlayer
 	GENERATED_BODY()
 
 public:
-	UFUNCTION(BlueprintCallable, Category = "Media Stream|Player")
+	UFUNCTION(Category = "Media Stream|Player")
 	virtual UMediaStream* GetMediaStream() const = 0;
 
 	/* Returns whether this player's controls can do anything. */
-	UFUNCTION(BlueprintCallable, CallInEditor, Category = "Media Stream|Player")
+	UFUNCTION(CallInEditor, Category = "Media Stream|Player")
 	virtual bool IsReadOnly() const = 0;
 
-	UFUNCTION(BlueprintCallable, Category = "Media Stream|Player")
+	UFUNCTION(Category = "Media Stream|Player")
 	virtual void OnCreated() = 0;
 
 	/**
@@ -52,59 +52,59 @@ public:
 	/**
 	 * @return Gets the media texture subobject.
 	 */
-	UFUNCTION(BlueprintCallable, Category = "Media Stream|Texture")
+	UFUNCTION(Category = "Media Stream|Texture")
 	virtual UMediaTexture* GetMediaTexture() const = 0;
 
 	/**
 	 * @return The current texture config.
 	 */
-	UFUNCTION(BlueprintCallable, Category = "Media Stream|Texture")
+	UFUNCTION(Category = "Media Stream|Texture")
 	virtual const FMediaStreamTextureConfig& GetTextureConfig() const = 0;
 
 	/**
 	 * Update the current texture's (and any newly set texture's) config.
 	 * @param InTextureConfig The new config.
 	 */
-	UFUNCTION(BlueprintCallable, Category = "Media Stream|Texture")
+	UFUNCTION(Category = "Media Stream|Texture")
 	virtual void SetTextureConfig(const FMediaStreamTextureConfig& InTextureConfig) = 0;
 
 	/**
 	 * Applies the current texture config to the current media texture.
 	 * This is generally automatic.
 	 */
-	UFUNCTION(BlueprintCallable, Category = "Media Stream|Texture")
+	UFUNCTION(Category = "Media Stream|Texture")
 	virtual void ApplyTextureConfig() = 0;
 
 	/**
 	 * @returns The active player, if there is one and it is valid.
 	 */
-	UFUNCTION(BlueprintCallable, Category = "Media Stream|Player")
+	UFUNCTION(Category = "Media Stream|Player")
 	virtual UMediaPlayer* GetPlayer() const = 0;
 
 	/**
 	 * @return True if there is an active and valid player.
 	 */
-	UFUNCTION(BlueprintCallable, Category = "Media Stream|Player")
+	UFUNCTION(Category = "Media Stream|Player")
 	virtual bool HasValidPlayer() const = 0;
 
 	/**
 	 * @return The current player config.
 	 */
-	UFUNCTION(BlueprintCallable, Category = "Media Stream|Player")
+	UFUNCTION(Category = "Media Stream|Player")
 	virtual const FMediaStreamPlayerConfig& GetPlayerConfig() const = 0;
 
 	/**
 	 * Update the current player's (and any newly set player's) config.
 	 * @param InPlayerConfig The new config.
 	 */
-	UFUNCTION(BlueprintCallable, Category = "Media Stream|Player")
+	UFUNCTION(Category = "Media Stream|Player")
 	virtual void SetPlayerConfig(const FMediaStreamPlayerConfig& InPlayerConfig) = 0;
 
 	/**
 	 * Applies the current player config to the current media player.
 	 * This is generally automatic.
 	 */
-	UFUNCTION(BlueprintCallable, Category = "Media Stream|Player")
+	UFUNCTION(Category = "Media Stream|Player")
 	virtual void ApplyPlayerConfig() = 0;
 
 	/**
@@ -112,13 +112,13 @@ public:
 	 * @param InIndex The index to start playing.
 	 * @return True on success. False on failure or invalid player.
 	 */
-	UFUNCTION(BlueprintCallable, Category = "Media Stream|Player")
+	UFUNCTION(Category = "Media Stream|Player")
 	virtual bool SetPlaylistIndex(int32 InIndex) = 0;
 
 	/**
 	 * @return The time offset in the media that has been requested.
 	 */
-	UFUNCTION(BlueprintCallable, Category = "Media Stream|Player")
+	UFUNCTION(Category = "Media Stream|Player")
 	virtual float GetRequestedSeekTime() const = 0;
 
 	/**
@@ -126,13 +126,13 @@ public:
 	 * @param InTime The time offset in the media to request.
 	 * @return True on success. False on failure or invalid player.
 	 */
-	UFUNCTION(BlueprintCallable, Category = "Media Stream|Player")
+	UFUNCTION(Category = "Media Stream|Player")
 	virtual bool SetRequestedSeekTime(float InTime) = 0;
 
 	/**
 	 * @return The frame offset in the media that has been requested.
 	 */
-	UFUNCTION(BlueprintCallable, Category = "Media Stream|Player")
+	UFUNCTION(Category = "Media Stream|Player")
 	virtual int32 GetRequestedSeekFrame() const = 0;
 
 	/**
@@ -140,13 +140,13 @@ public:
 	 * @param InFrame The frame offset in the media to request.
 	 * @return True on success. False on failure or invalid player.
 	 */
-	UFUNCTION(BlueprintCallable, Category = "Media Stream|Player")
+	UFUNCTION(Category = "Media Stream|Player")
 	virtual bool SetRequestedSeekFrame(int32 InFrame) = 0;
 
 	/**
 	 * @return Gets the media player playback state.
 	 */
-	UFUNCTION(BlueprintCallable, Category = "Media Stream|Player")
+	UFUNCTION(Category = "Media Stream|Player")
 	virtual EMediaStreamPlaybackState GetPlaybackState() const = 0;
 
 	/**
@@ -154,81 +154,81 @@ public:
 	 * @param InState The state to change to.
 	 * @return True on success. False on failure or invalid player.
 	 */
-	UFUNCTION(BlueprintCallable, Category = "Media Stream|Player")
+	UFUNCTION(Category = "Media Stream|Player")
 	virtual bool SetPlaybackState(EMediaStreamPlaybackState InState) = 0;
 
 	/**
 	 * @return Gets the active playlist index. May not match the actual player. -1 on error.
 	 */
-	UFUNCTION(BlueprintCallable, Category = "Media Stream|Player")
+	UFUNCTION(Category = "Media Stream|Player")
 	virtual int32 GetPlaylistIndex() const = 0;
 
 	/**
 	 * @return The size of the current playlist. -1 on error.
 	 */
-	UFUNCTION(BlueprintCallable, Category = "Media Stream|Player")
+	UFUNCTION(Category = "Media Stream|Player")
 	virtual int32 GetPlaylistNum() const = 0;
 
 	/**
 	 * @return For proxy streams, returns the stream this is getting data from.
 	 */
-	UFUNCTION(BlueprintCallable, Category = "Media Stream|Player")
+	UFUNCTION(Category = "Media Stream|Player")
 	virtual UMediaStream* GetSourceStream() const = 0;
 
 	/**
 	 * Opens the source if it is not already opened.
 	 * @return True if successful.
 	 */
-	UFUNCTION(BlueprintCallable, CallInEditor, Category = "Media Stream|Player")
+	UFUNCTION(CallInEditor, Category = "Media Stream|Player")
 	virtual bool OpenSource() = 0;
 
 	/**
 	 * Continues play of the current media.
 	 * @return True if successful.
 	 */
-	UFUNCTION(BlueprintCallable, CallInEditor, Category = "Media Stream|Player")
+	UFUNCTION(CallInEditor, Category = "Media Stream|Player")
 	virtual bool Play() = 0;
 
 	/**
 	 * Pauses playback of the current media.
 	 * @return True if successful.
 	 */
-	UFUNCTION(BlueprintCallable, CallInEditor, Category = "Media Stream|Player")
+	UFUNCTION(CallInEditor, Category = "Media Stream|Player")
 	virtual bool Pause() = 0;
 
 	/**
 	 * Seeks to the start of the current media.
 	 * @return True if successful.
 	 */
-	UFUNCTION(BlueprintCallable, CallInEditor, Category = "Media Stream|Player")
+	UFUNCTION(CallInEditor, Category = "Media Stream|Player")
 	virtual bool Rewind() = 0;
 
 	/**
 	 * Seeks to the end of the current media.
 	 * @return True if successful.
 	 */
-	UFUNCTION(BlueprintCallable, CallInEditor, Category = "Media Stream|Player")
+	UFUNCTION(CallInEditor, Category = "Media Stream|Player")
 	virtual bool FastForward() = 0;
 
 	/**
 	 * Goes to the previous item in the playlist. Only possible with playlist sources.
 	 * @return True if successful.
 	 */
-	UFUNCTION(BlueprintCallable, CallInEditor, Category = "Media Stream|Player")
+	UFUNCTION(CallInEditor, Category = "Media Stream|Player")
 	virtual bool Previous() = 0;
 
 	/**
 	 * Goes to the next item in the playlist. Only possible with playlist sources.
 	 * @return True if successful.
 	 */
-	UFUNCTION(BlueprintCallable, CallInEditor, Category = "Media Stream|Player")
+	UFUNCTION(CallInEditor, Category = "Media Stream|Player")
 	virtual bool Next() = 0;
 
 	/**
 	 * Closes the current media player.
 	 * @return True if successful.
 	 */
-	UFUNCTION(BlueprintCallable, CallInEditor, Category = "Media Stream|Player")
+	UFUNCTION(CallInEditor, Category = "Media Stream|Player")
 	virtual bool Close() = 0;
 
 	/**
diff --git a/Engine/Plugins/Experimental/MediaStream/Source/MediaStream/Public/MediaStream.h b/Engine/Plugins/Experimental/MediaStream/Source/MediaStream/Public/MediaStream.h
--- a/Engine/Plugins/Experimental/MediaStream/Source/MediaStream/Public/MediaStream.h
+++ b/Engine/Plugins/Experimental/MediaStream/Source/MediaStream/Public/MediaStream.h
@@ -30,20 +30,20 @@ public:
 	/**
 	 * @return True if the media source scheme is not empty. Does not guarantee the source is correct, only valid.
 	 */
-	UFUNCTION(BlueprintPure, Category = "Media Stream")
+	UFUNCTION(Category = "Media Stream")
 	MEDIASTREAM_API bool HasValidSource() const;
 
 	/**
 	 * @return The current Media Stream Source. Does not guaranteed that it is valid.
 	 */
-	UFUNCTION(BlueprintPure, Category = "Media Stream")
+	UFUNCTION(Category = "Media Stream")
 	MEDIASTREAM_API const FMediaStreamSource& GetSource() const;
 
 	/**
 	 * @return The source at the end of the chain of proxy players, or the local one if it is local. If
 	 * the chain is interrupted, the last link in the chain's source will be returned.
 	 */
-	UFUNCTION(BlueprintPure, Category = "Media Stream")
+	UFUNCTION(Category = "Media Stream")
 	MEDIASTREAM_API const FMediaStreamSource& ResolveSource() const;
 
 	/**
@@ -51,13 +51,13 @@ public:
 	 * @param InSource The new source.
 	 * @return True if the source was successfully applied.
 	 */
-	UFUNCTION(BlueprintCallable, Category = "Media Stream")
+	UFUNCTION(Category = "Media Stream")
 	MEDIASTREAM_API bool SetSource(const FMediaStreamSource& InSource);
 
 	/**
 	 * @return Gets the active Media Stream Player.
 	 */
-	UFUNCTION(BlueprintPure, Category = "Media Stream")
+	UFUNCTION(Category = "Media Stream")
 	MEDIASTREAM_API TScriptInterface<IMediaStreamPlayer> GetPlayer() const;
 
 	/**
@@ -65,7 +65,7 @@ public:
 	 * @param  bInForceRecreatePlayer Will remove and recreate the player.
 	 * @return True if a valid player exists.
 	 */
-	UFUNCTION(BlueprintCallable, Category = "Media Stream")
+	UFUNCTION(Category = "Media Stream")
 	MEDIASTREAM_API bool EnsurePlayer(bool bInForceRecreatePlayer = false);
 
 	/**
diff --git a/Engine/Plugins/Experimental/MediaStream/Source/MediaStream/Public/MediaStreamObjectHandlerSubsystem.h b/Engine/Plugins/Experimental/MediaStream/Source/MediaStream/Public/MediaStreamObjectHandlerSubsystem.h
--- a/Engine/Plugins/Experimental/MediaStream/Source/MediaStream/Public/MediaStreamObjectHandlerSubsystem.h
+++ b/Engine/Plugins/Experimental/MediaStream/Source/MediaStream/Public/MediaStreamObjectHandlerSubsystem.h
@@ -22,7 +22,7 @@ public:
 	/**
 	 * @return Gets the instance of this subsystem.
 	 */
-	UFUNCTION(BlueprintPure, Category = "Media Stream|Object")
+	UFUNCTION(Category = "Media Stream|Object")
 	static MEDIASTREAM_API UMediaStreamObjectHandlerSubsystem* Get();
 
 	/**
@@ -30,7 +30,7 @@ public:
 	 * @param InClass The class to check.
 	 * @return True if any handler can handle the class.
 	 */
-	UFUNCTION(BlueprintPure, Category = "Media Stream|Object")
+	UFUNCTION(Category = "Media Stream|Object")
 	MEDIASTREAM_API bool CanHandleObject(const UClass* InClass) const;
 
 	/**
@@ -38,7 +38,7 @@ public:
 	 * @param InClass The class to check. Must be an exact match.
 	 * @return True if the handler is registered.
 	 */
-	UFUNCTION(BlueprintPure, Category = "Media Stream|Object")
+	UFUNCTION(Category = "Media Stream|Object")
 	MEDIASTREAM_API bool HasObjectHandler(const UClass* InClass) const;
 
 	/**
@@ -46,6 +46,6 @@ public:
 	 * Note: This usually means loading the media source. @see bCanLoadSource.
 	 * @return A player or nullptr.
 	 */
-	UFUNCTION(BlueprintCallable, Category = "Media Stream|Object")
+	UFUNCTION(Category = "Media Stream|Object")
 	MEDIASTREAM_API UMediaPlayer* CreateMediaPlayer(const FMediaStreamObjectHandlerCreatePlayerParams& InParams) const;
 };
diff --git a/Engine/Plugins/Experimental/MediaStream/Source/MediaStream/Public/MediaStreamSchemeHandlerSubsystem.h b/Engine/Plugins/Experimental/MediaStream/Source/MediaStream/Public/MediaStreamSchemeHandlerSubsystem.h
--- a/Engine/Plugins/Experimental/MediaStream/Source/MediaStream/Public/MediaStreamSchemeHandlerSubsystem.h
+++ b/Engine/Plugins/Experimental/MediaStream/Source/MediaStream/Public/MediaStreamSchemeHandlerSubsystem.h
@@ -26,7 +26,7 @@ public:
 	/**
 	 * @return Gets the instance of this subsystem.
 	 */
-	UFUNCTION(BlueprintPure, Category = "Media Stream|Scheme")
+	UFUNCTION(Category = "Media Stream|Scheme")
 	static MEDIASTREAM_API UMediaStreamSchemeHandlerSubsystem* Get();
 
 	/**
@@ -34,13 +34,13 @@ public:
 	 * @param InScheme The scheme to check.
 	 * @return True if the handler is registered.
 	 */
-	UFUNCTION(BlueprintPure, Category = "Media Stream|Scheme")
+	UFUNCTION(Category = "Media Stream|Scheme")
 	MEDIASTREAM_API bool HasSchemeHandler(FName InScheme) const;
 
 	/**
 	 * @return Gets the list of registered scheme handlers.
 	 */
-	UFUNCTION(BlueprintPure, Category = "Media Stream|Scheme")
+	UFUNCTION(Category = "Media Stream|Scheme")
 	MEDIASTREAM_API TArray<FName> GetSchemeHandlerNames() const;
 
 	/**
@@ -50,7 +50,7 @@ public:
 	 * @param InPath The path to the media source. This should not be Url encoded.
 	 * @return A valid stream source object or none if it was not a valid Url.
 	 */
-	UFUNCTION(BlueprintCallable, Category = "Media Stream|Scheme")
+	UFUNCTION(Category = "Media Stream|Scheme")
 	FMediaStreamSource CreateSource(UObject* InOuter, FName InScheme, const FString& InPath) const;
 
 	/**
@@ -58,6 +58,6 @@ public:
 	 * Note: This usually means loading the media source. @see bCanLoadSource.
 	 * @return A player or nullptr.
 	 */
-	UFUNCTION(BlueprintCallable, Category = "Media Stream|Scheme")
+	UFUNCTION(Category = "Media Stream|Scheme")
 	MEDIASTREAM_API UMediaPlayer* CreateOrUpdatePlayer(const FMediaStreamSchemeHandlerCreatePlayerParams& InParams) const;
 };
diff --git a/Engine/Plugins/Experimental/MediaStream/Source/MediaStream/Public/MediaStreamSourceBlueprintLibrary.h b/Engine/Plugins/Experimental/MediaStream/Source/MediaStream/Public/MediaStreamSourceBlueprintLibrary.h
--- a/Engine/Plugins/Experimental/MediaStream/Source/MediaStream/Public/MediaStreamSourceBlueprintLibrary.h
+++ b/Engine/Plugins/Experimental/MediaStream/Source/MediaStream/Public/MediaStreamSourceBlueprintLibrary.h
@@ -26,7 +26,7 @@ public:
 	 * @param InSource The source to check.
 	 * @return True if the source is valid. Does not guarantee it is correct, only valid.
 	 */
-	UFUNCTION(BlueprintPure, Category = "Media Stream|Source")
+	UFUNCTION(Category = "Media Stream|Source")
 	static MEDIASTREAM_API bool IsValidMediaSource(const FMediaStreamSource& InSource);
 
 	/**
@@ -34,7 +34,7 @@ public:
 	 * @param InAsset The asset to check.
 	 * @return True if the asset is valid.
 	 */
-	UFUNCTION(BlueprintPure, Category = "Media Stream|Source")
+	UFUNCTION(Category = "Media Stream|Source")
 	static MEDIASTREAM_API bool IsAssetValid(const TSoftObjectPtr<UObject>& InAsset);
 
 	/**
@@ -42,7 +42,7 @@ public:
 	 * @param InPath The asset path to check.
 	 * @return True if the asset is valid.
 	 */
-	UFUNCTION(BlueprintPure, Category = "Media Stream|Source")
+	UFUNCTION(Category = "Media Stream|Source")
 	static MEDIASTREAM_API bool IsAssetPathValid(const FString& InPath);
 
 	/**
@@ -50,7 +50,7 @@ public:
 	 * @param InPath The asset soft path to check.
 	 * @return True if the asset is valid.
 	 */
-	UFUNCTION(BlueprintPure, Category = "Media Stream|Source")
+	UFUNCTION(Category = "Media Stream|Source")
 	static MEDIASTREAM_API bool IsAssetSoftPathValid(const FSoftObjectPath& InPath);
 
 	/**
@@ -60,7 +60,7 @@ public:
 	 * @param InPath The path to the source.
 	 * @return A valid stream source object or none if it was not a valid object.
 	 */
-	UFUNCTION(BlueprintCallable, Category = "Media Stream|Source")
+	UFUNCTION(Category = "Media Stream|Source")
 	static MEDIASTREAM_API FMediaStreamSource MakeMediaSourceFromSchemePath(UMediaStream* InMediaStream, FName InScheme, const FString& InPath);
 
 	/**
@@ -69,7 +69,7 @@ public:
 	 * @param InObject The playable object.
 	 * @return A valid stream source object or none if it was not a valid object.
 	 */
-	UFUNCTION(BlueprintCallable, Category = "Media Stream|Source")
+	UFUNCTION(Category = "Media Stream|Source")
 	static MEDIASTREAM_API FMediaStreamSource MakeMediaSourceFromAsset(UMediaStream* InMediaStream, const TSoftObjectPtr<UObject>& InAsset);
 
 	/**
@@ -78,7 +78,7 @@ public:
 	 * @param InStreamName The name of the registered stream.
 	 * @return A valid stream source object or none if it was not a valid stream or the stream was not found.
 	 */
-	UFUNCTION(BlueprintCallable, Category = "Media Stream|Source")
+	UFUNCTION(Category = "Media Stream|Source")
 	static MEDIASTREAM_API FMediaStreamSource MakeMediaSourceFromStreamName(UMediaStream* InMediaStream, FName InStreamName);
 
 	/**
@@ -87,7 +87,7 @@ public:
 	 * @param InFileName The path to the file name. Should be a full path or relative to the project or engine root.
 	 * @return A valid stream source object or none if the file was not found.
 	 */
-	UFUNCTION(BlueprintCallable, Category = "Media Stream|Source")
+	UFUNCTION(Category = "Media Stream|Source")
 	static MEDIASTREAM_API FMediaStreamSource MakeMediaSourceFromFile(UMediaStream* InMediaStream, const FString& InFileName);
 
 	/**
@@ -97,7 +97,7 @@ public:
 	 * @param InObject The playable object. Should share a common ancestor with the Media Stream, but this is not enforced.
 	 * @return A valid stream source object or none if the file was not found.
 	 */
-	UFUNCTION(BlueprintCallable, Category = "Media Stream|Source")
+	UFUNCTION(Category = "Media Stream|Source")
 	static MEDIASTREAM_API FMediaStreamSource MakeMediaSourceFromSubobject(UMediaStream* InMediaStream, UObject* InObject);
 
 	/**
@@ -107,7 +107,7 @@ public:
 	 * @param InOuter The outer of the new object. Defaults to the media stream object.
 	 * @return A valid stream source object or none if the file was not found.
 	 */
-	UFUNCTION(BlueprintCallable, Category = "Media Stream|Source")
+	UFUNCTION(Category = "Media Stream|Source")
 	static MEDIASTREAM_API FMediaStreamSource MakeMediaSourceFromSubobjectClass(UMediaStream* InMediaStream, const UClass* InClass);
 
 	/**
