[MediaStream] Fix seek conversion to timespan from seconds.

#Synced-CL 48520225

diff --git a/Engine/Plugins/Experimental/MediaStream/Source/MediaStream/Private/Players/MediaStreamLocalPlayer.cpp b/Engine/Plugins/Experimental/MediaStream/Source/MediaStream/Private/Players/MediaStreamLocalPlayer.cpp
--- a/Engine/Plugins/Experimental/MediaStream/Source/MediaStream/Private/Players/MediaStreamLocalPlayer.cpp
+++ b/Engine/Plugins/Experimental/MediaStream/Source/MediaStream/Private/Players/MediaStreamLocalPlayer.cpp
@@ -392,7 +392,7 @@ bool UMediaStreamLocalPlayer::ApplyRequestedSeekFrame()
 
 		if (FrameRate > 0.f)
 		{
-			bSuccess = Player->Seek(static_cast<float>(RequestedSeekFrame) / FrameRate);
+			bSuccess = Player->Seek(FTimespan::FromSeconds(static_cast<float>(RequestedSeekFrame) / FrameRate));
 		}
 	}
 
