AvalancheTests: Automate 'Material Designer - The media layer can be added'

#rnx

#Synced-CL 45465245

diff --git a/Engine/Plugins/Experimental/MediaStream/Source/MediaStream/Private/SchemeHandlers/MediaStreamFileSchemeHandler.cpp b/Engine/Plugins/Experimental/MediaStream/Source/MediaStream/Private/SchemeHandlers/MediaStreamFileSchemeHandler.cpp
--- a/Engine/Plugins/Experimental/MediaStream/Source/MediaStream/Private/SchemeHandlers/MediaStreamFileSchemeHandler.cpp
+++ b/Engine/Plugins/Experimental/MediaStream/Source/MediaStream/Private/SchemeHandlers/MediaStreamFileSchemeHandler.cpp
@@ -101,7 +101,8 @@ FString FMediaStreamFileSchemeHandler::ResolveFilePath(const FString& InPath)
 		return FilePath;
 	}
 
-	UE_LOG(LogMediaStream, Error, TEXT("Invalid File Path in FMediaStreamFileSchemeHandler::ResolveFilePath [%s]"), *InPath);
+	// Shouldn't generate an error since it's default behavior and common code path. Otherwise it breaks tests.
+	UE_LOGF(LogMediaStream, Display, "Invalid File Path in FMediaStreamFileSchemeHandler::ResolveFilePath [%ls]", *InPath);
 
 	return FString();
 }
diff --git a/Engine/Plugins/VirtualProduction/Avalanche/Source/AvalancheShapes/Public/AvaShapeMesh.h b/Engine/Plugins/VirtualProduction/Avalanche/Source/AvalancheShapes/Public/AvaShapeMesh.h
--- a/Engine/Plugins/VirtualProduction/Avalanche/Source/AvalancheShapes/Public/AvaShapeMesh.h
+++ b/Engine/Plugins/VirtualProduction/Avalanche/Source/AvalancheShapes/Public/AvaShapeMesh.h
@@ -6,6 +6,7 @@
 #include "AvaShapesDefs.h"
 #include "AvaShapeUVParameters.h"
 #include "Containers/Array.h"
+#include "Materials/MaterialInterface.h"
 #include "Math/MathFwd.h"
 #include "AvaShapeMesh.generated.h"
 
diff --git a/Engine/Source/Developer/DesktopWidgets/Private/Widgets/Input/SFilePathPicker.cpp b/Engine/Source/Developer/DesktopWidgets/Private/Widgets/Input/SFilePathPicker.cpp
--- a/Engine/Source/Developer/DesktopWidgets/Private/Widgets/Input/SFilePathPicker.cpp
+++ b/Engine/Source/Developer/DesktopWidgets/Private/Widgets/Input/SFilePathPicker.cpp
@@ -65,6 +65,10 @@ void SFilePathPicker::Construct( const FArguments& InArgs )
 	];
 }
 
+void SFilePathPicker::SimulateTextInput(const FText& NewText)
+{
+	HandleTextBoxTextCommitted(NewText, ETextCommit::Default);
+}
 
 /* SFilePathPicker callbacks
  *****************************************************************************/
diff --git a/Engine/Source/Developer/DesktopWidgets/Public/Widgets/Input/SFilePathPicker.h b/Engine/Source/Developer/DesktopWidgets/Public/Widgets/Input/SFilePathPicker.h
--- a/Engine/Source/Developer/DesktopWidgets/Public/Widgets/Input/SFilePathPicker.h
+++ b/Engine/Source/Developer/DesktopWidgets/Public/Widgets/Input/SFilePathPicker.h
@@ -82,6 +82,9 @@ public:
 	 */
 	UE_API void Construct( const FArguments& InArgs );
 
+	/** Method which simulates the text input in the path text box. */
+	UE_API void SimulateTextInput(const FText& NewText);
+
 private:
 
 	/** Callback for clicking the browse button. */
